[EaseCurveTool] Bring the Ease Curve Tool to Sequencer and Curve Editor
- Move from MotionDesign plugin to its own plugin: EaseCurveTool
- Support setting and creating presets from curve editor

#virtualized

diff --git a/Engine/Plugins/Experimental/EaseCurveTool/Source/EaseCurveTool/Private/EaseCurveKeyCache.h b/Engine/Plugins/Experimental/EaseCurveTool/Source/EaseCurveTool/Private/EaseCurveKeyCache.h
new file mode 100644
--- /dev/null
+++ b/Engine/Plugins/Experimental/EaseCurveTool/Source/EaseCurveTool/Private/EaseCurveKeyCache.h
@@ -0,0 +1,98 @@
+﻿// Copyright Epic Games, Inc. All Rights Reserved.
+
+#pragma once
+
+#include "Channels/ChannelCurveModel.h"
+#include "Channels/MovieSceneChannelData.h"
+#include "Channels/MovieSceneChannelHandle.h"
+#include "CurveModel.h"
+#include "EaseCurveTangents.h"
+#include "MovieSceneSection.h"
+
+namespace UE::EaseCurveTool
+{
+
+struct FEaseCurveKeyCache
+{
+public:
+	TVariant<TObjectPtr<UMovieSceneSection>, FCurveModel*> ObjectToModify;
+	FMovieSceneChannelHandle ChannelHandle;
+	TArray<FKeyHandle> KeyHandles;
+
+	template <class InChannelHandleType, class InChannelValueType>
+	TMovieSceneChannelData<InChannelValueType> GetChannelData()
+	{
+		const TMovieSceneChannelHandle<InChannelHandleType> Channel = ChannelHandle.Cast<InChannelHandleType>();
+		return Channel.Get()->GetData();
+	}
+
+	bool CanEaseCurveKey(const FKeyHandle InKeyHandle) const
+	{
+		if (const TObjectPtr<UMovieSceneSection> Section = ObjectToModify.Get<TObjectPtr<UMovieSceneSection>>())
+		{
+			const FName ChannelTypeName = ChannelHandle.GetChannelTypeName();
+			if (ChannelTypeName == FMovieSceneDoubleChannel::StaticStruct()->GetFName())
+			{
+				if (!HasWeightedBrokenTangents<FMovieSceneDoubleChannel, FMovieSceneDoubleValue>(InKeyHandle, ChannelHandle))
+				{
+					return false;
+				}
+			}
+			else if (ChannelTypeName == FMovieSceneFloatChannel::StaticStruct()->GetFName())
+			{
+				if (!HasWeightedBrokenTangents<FMovieSceneFloatChannel, FMovieSceneFloatValue>(InKeyHandle, ChannelHandle))
+				{
+					return false;
+				}
+			}
+		}
+		else if (const FCurveModel* CurveModel = ObjectToModify.Get<FCurveModel*>())
+		{
+			if (const FChannelCurveModel<FMovieSceneFloatChannel, float, float>* FloatChannelModel = 
+				static_cast<const FChannelCurveModel<FMovieSceneFloatChannel, float, float>*>(CurveModel))
+			{
+				if (!HasWeightedBrokenTangents<FMovieSceneDoubleChannel, FMovieSceneDoubleValue>(InKeyHandle, ChannelHandle))
+				{
+					return false;
+				}
+			}
+			else if (const FChannelCurveModel<FMovieSceneDoubleChannel, double, double>* DoubleChannelModel = 
+				static_cast<const FChannelCurveModel<FMovieSceneDoubleChannel, double, double>*>(CurveModel))
+			{
+				if (!HasWeightedBrokenTangents<FMovieSceneDoubleChannel, FMovieSceneDoubleValue>(InKeyHandle, ChannelHandle))
+				{
+					return false;
+				}
+			}
+		}
+
+		return true;
+	}
+
+	template <class InChannelType, class InValueType>
+	static bool HasWeightedBrokenTangents(const FKeyHandle& InKeyHandle, const FMovieSceneChannelHandle& InChannelHandle)
+	{
+		InChannelType* const Channel = static_cast<InChannelType*>(InChannelHandle.Get());
+		if (!Channel)
+		{
+			return false;
+		}
+
+		const int32 KeyIndex = Channel->GetIndex(InKeyHandle);
+		if (KeyIndex == INDEX_NONE)
+		{
+			return false;
+		}
+
+		const TArrayView<const InValueType> ChannelValues = Channel->GetValues();
+
+		if (!FEaseCurveTangents::CanEaseCurveKey(ChannelValues[KeyIndex]))
+		{
+			return false;
+		}
+
+		return true;
+	}
+};
+
+} // namespace UE::EaseCurveTool
diff --git a/Engine/Plugins/Experimental/EaseCurveTool/Source/EaseCurveTool/Private/EaseCurveKeySelection.cpp b/Engine/Plugins/Experimental/EaseCurveTool/Source/EaseCurveTool/Private/EaseCurveKeySelection.cpp
new file mode 100644
--- /dev/null
+++ b/Engine/Plugins/Experimental/EaseCurveTool/Source/EaseCurveTool/Private/EaseCurveKeySelection.cpp
@@ -0,0 +1,488 @@
+﻿// Copyright Epic Games, Inc. All Rights Reserved.
+
+#include "EaseCurveKeySelection.h"
+#include "Channels/ChannelCurveModel.h"
+#include "Channels/MovieSceneChannel.h"
+#include "Channels/MovieSceneDoubleChannel.h"
+#include "Channels/MovieSceneFloatChannel.h"
+#include "CurveEditor.h"
+#include "EaseCurveTangents.h"
+#include "EaseCurveTool.h"
+#include "EaseCurveToolUtils.h"
+#include "IKeyArea.h"
+#include "MVVM/Selection/Selection.h"
+#include "MVVM/ViewModels/ChannelModel.h"
+#include "MVVM/ViewModels/SequencerEditorViewModel.h"
+#include "ScopedTransaction.h"
+#include "MVVM/ViewModels/SequenceModel.h"
+
+using namespace UE::Sequencer;
+
+namespace UE::EaseCurveTool
+{
+
+template <class InChannelType, class InChannelValue>
+void NormalizeChannelValues(const FKeyHandle& InKeyHandle, const FKeyHandle& InNextKeyHandle
+	, const FMovieSceneChannelHandle& InChannelHandle
+	, const bool bInAutoFlipTangents
+	, const FFrameRate& InDisplayRate, const FFrameRate& InTickResolution
+	, TArray<FEaseCurveTangents>& OutKeySetTangents
+	, TArray<FEaseCurveTangents>& OutChangingTangents)
+{
+	TMovieSceneChannelHandle<InChannelType> Channel = InChannelHandle.Cast<InChannelType>();
+	TMovieSceneChannelData<InChannelValue> ChannelData = Channel.Get()->GetData();
+
+	const TArrayView<InChannelValue> ChannelValues = ChannelData.GetValues();
+	const TArrayView<FFrameNumber> ChannelTimes = ChannelData.GetTimes();
+
+	const int32 KeyIndex = ChannelData.GetIndex(InKeyHandle);
+	const int32 NextKeyIndex = ChannelData.GetIndex(InNextKeyHandle);
+
+	// If there is a key frame after this key frame that we are editing, we check if the that key frame value is less
+	// than or greater than this key frame value. If less, flip the tangent (if option is set).
+	const bool bIncreasingValue = ChannelValues[NextKeyIndex].Value >= ChannelValues[KeyIndex].Value;
+
+	FEaseCurveTangents Tangents = FEaseCurveTangents(ChannelValues[KeyIndex], ChannelValues[NextKeyIndex]);
+
+	if (bInAutoFlipTangents && !bIncreasingValue)
+	{
+		Tangents.Start *= -1.f;
+		Tangents.End *= -1.f;
+	}
+
+	// Scale time/value to normalized tangent range
+	FEaseCurveTangents ScaledTangents = Tangents;
+	ScaledTangents.Normalize(ChannelTimes[KeyIndex], ChannelValues[KeyIndex].Value
+		, ChannelTimes[NextKeyIndex], ChannelValues[NextKeyIndex].Value
+		, InDisplayRate, InTickResolution);
+
+	OutKeySetTangents.Add(ScaledTangents);
+	if (ChannelValues[KeyIndex].Value != ChannelValues[NextKeyIndex].Value)
+	{
+		OutChangingTangents.Add(ScaledTangents);
+	}
+}
+
+template <class InChannelType, class InValueType>
+void SetChannelValues(const FEaseCurveTangents& InTangents
+	, const EEaseCurveToolOperation InOperation
+	, const FKeyHandle& InKeyHandle, const FKeyHandle& InNextKeyHandle
+	, const FMovieSceneChannelHandle& InChannelHandle
+	, const bool bInAutoFlipTangents
+	, const FFrameRate& InDisplayRate, const FFrameRate& InTickResolution
+	, const TVariant<TObjectPtr<UMovieSceneSection>, FCurveModel*>& InObjectToModify)
+{
+	TMovieSceneChannelHandle<InChannelType> Channel = /*InChannelData.ChannelHandle*/InChannelHandle.Cast<InChannelType>();
+	TMovieSceneChannelData<InValueType> ChannelData = Channel.Get()->GetData();
+	TMovieSceneChannelHandle<FMovieSceneDoubleChannel> DoubleChannel;
+	TArrayView<InValueType> ChannelValues = ChannelData.GetValues();
+	const TArrayView<FFrameNumber> ChannelTimes = ChannelData.GetTimes();
+
+	const int32 KeyIndex = ChannelData.GetIndex(InKeyHandle);
+	const int32 NextKeyIndex = ChannelData.GetIndex(InNextKeyHandle);
+
+	FEaseCurveTangents ScaledTangents = InTangents;
+
+	// If there is a key frame after this key frame that we are editing, we check if the that key frame value is less
+	// than or greater than this key frame value. If less, flip the tangent (if option is set).
+	const bool bIncreasingValue = ChannelValues[NextKeyIndex].Value >= ChannelValues[KeyIndex].Value;
+
+	if (bInAutoFlipTangents && !bIncreasingValue)
+	{
+		ScaledTangents.Start *= -1.f;
+		ScaledTangents.End *= -1.f;
+	}
+
+	// Scale normalized tangents to time/value range
+	ScaledTangents.ScaleUp(ChannelTimes[KeyIndex], ChannelValues[KeyIndex].Value
+		, ChannelTimes[NextKeyIndex], ChannelValues[NextKeyIndex].Value
+		, InDisplayRate, InTickResolution);
+
+	const FScopedTransaction Transaction(NSLOCTEXT("EaseCurveTool", "SetSequencerCurveTangents", "Set Sequencer Curve Tangents"));
+
+	// Mark the section as changed to force the curve editor to update its display
+	if (InObjectToModify.IsType<TObjectPtr<UMovieSceneSection>>())
+	{
+		if (const TObjectPtr<UMovieSceneSection> SectionToModify = InObjectToModify.Get<TObjectPtr<UMovieSceneSection>>())
+		{
+			SectionToModify->Modify();
+			SectionToModify->MarkAsChanged();
+		}
+	}
+	else if (InObjectToModify.IsType<FCurveModel*>())
+	{
+		if (FCurveModel* const CurveModelToModify = InObjectToModify.Get<FCurveModel*>())
+		{
+			CurveModelToModify->Modify();
+			//CurveModelToModify->MarkAsChanged();
+		}
+	}
+
+	// Set this keys leave tangent
+	if (InOperation == EEaseCurveToolOperation::Out || InOperation == EEaseCurveToolOperation::InOut)
+	{
+		ChannelValues[KeyIndex].InterpMode = ERichCurveInterpMode::RCIM_Cubic;
+		ChannelValues[KeyIndex].Tangent.TangentWeightMode = ERichCurveTangentWeightMode::RCTWM_WeightedBoth;
+		ChannelValues[KeyIndex].TangentMode = ERichCurveTangentMode::RCTM_Break;
+		ChannelValues[KeyIndex].Tangent.LeaveTangent = ScaledTangents.Start;
+		ChannelValues[KeyIndex].Tangent.LeaveTangentWeight = ScaledTangents.StartWeight;
+	}
+
+	// Set the next keys arrive tangent
+	if (NextKeyIndex != INDEX_NONE && (InOperation == EEaseCurveToolOperation::In || InOperation == EEaseCurveToolOperation::InOut))
+	{
+		ChannelValues[NextKeyIndex].InterpMode = ERichCurveInterpMode::RCIM_Cubic;
+		ChannelValues[NextKeyIndex].Tangent.TangentWeightMode = ERichCurveTangentWeightMode::RCTWM_WeightedBoth;
+		ChannelValues[NextKeyIndex].TangentMode = ERichCurveTangentMode::RCTM_Break;
+		ChannelValues[NextKeyIndex].Tangent.ArriveTangent = ScaledTangents.End;
+		ChannelValues[NextKeyIndex].Tangent.ArriveTangentWeight = ScaledTangents.EndWeight;
+	}
+}
+
+FEaseCurveKeySelection::FEaseCurveKeySelection(const TSharedPtr<ISequencer>& InSequencer)
+{
+	if (!InSequencer.IsValid())
+	{
+		return;
+	}
+
+	const TSharedPtr<FSequencerEditorViewModel> SequencerViewModel = InSequencer->GetViewModel();
+	if (!SequencerViewModel.IsValid())
+	{
+		return;
+	}
+
+	const TSharedPtr<FSequencerSelection> SequencerSelection = SequencerViewModel->GetSelection();
+	if (!SequencerSelection.IsValid())
+	{
+		return;
+	}
+
+	bAreAllEaseCurves = true;
+
+	bool bCurveEditorSelection = false;
+
+	if (const TSharedPtr<FCurveEditor> CurveEditor = GetCurveEditorFromSequencer(InSequencer.ToSharedRef()))
+	{
+		const FCurveEditorSelection& CurveEditorSelection = CurveEditor->GetSelection();
+		const TMap<FCurveModelID, FKeyHandleSet> SelectedKeys = CurveEditorSelection.GetAll();
+		if (!SelectedKeys.IsEmpty())
+		{
+			const TMap<FCurveModelID, TUniquePtr<FCurveModel>>& Curves = CurveEditor->GetCurves();
+
+			for (const TPair<FCurveModelID, FKeyHandleSet>& SelectedKeyPair : SelectedKeys)
+			{
+				for (const FKeyHandle& SelectedKey : SelectedKeyPair.Value.AsArray())
+				{
+					if (Curves.Contains(SelectedKeyPair.Key))
+					{
+						if (FCurveModel* const CurveModel = Curves[SelectedKeyPair.Key].Get())
+						{
+							if (!ProcessCurveEditorKeySelection(CurveModel, SelectedKey))
+							{
+								bAreAllEaseCurves = false;
+								break;
+							}
+							else
+							{
+								bCurveEditorSelection = true;
+							}
+						}
+					}
+				}
+			}
+		}
+	}
+
+	if (!bCurveEditorSelection)
+	{
+		for (const FKeyHandle Key : SequencerSelection->KeySelection)
+		{
+			if (!ProcessSequencerKeySelection(SequencerSelection, Key))
+			{
+				bAreAllEaseCurves = false;
+				break;
+			}
+		}
+	}
+}
+
+bool FEaseCurveKeySelection::ProcessSequencerKeySelection(const TSharedPtr<FSequencerSelection>& InSequencerSelection, const FKeyHandle InKey)
+{
+	if (!InSequencerSelection.IsValid() || InKey == FKeyHandle::Invalid())
+	{
+		return false;
+	}
+
+	const TViewModelPtr<FChannelModel> ChannelModel = InSequencerSelection->KeySelection.GetModelForKey(InKey);
+	if (!ChannelModel.IsValid())
+	{
+		return false;
+	}
+
+	const TSharedPtr<IKeyArea> KeyArea = ChannelModel->GetKeyArea();
+	if (!KeyArea.IsValid())
+	{
+		return false;
+	}
+
+	const FMovieSceneChannelHandle& ChannelHandle = KeyArea->GetChannel();
+
+	const FString PathName = IOutlinerExtension::GetPathName(ChannelModel);
+	const FString ChannelName = ChannelModel->GetChannelName().ToString();
+	const FName FullPathMapKey = FName(PathName + ChannelName);
+
+	FEaseCurveKeyCache& Entry = SequencerKeyData.FindOrAdd(FullPathMapKey);
+	Entry.ObjectToModify.Emplace<TObjectPtr<UMovieSceneSection>>(ChannelModel->GetSection());
+	Entry.ChannelHandle = ChannelHandle;
+	Entry.KeyHandles.Add(InKey);
+
+	TotalSelectedKeys++;
+
+	const int32 AllKeyCount = ChannelHandle.Get()->GetNumKeys();
+	const int32 SelectedKeyCount = Entry.KeyHandles.Num();
+
+	if (SelectedKeyCount == 1 && TotalSelectedKeys == 1)
+	{
+		const int32 KeyIndex = ChannelHandle.Get()->GetIndex(Entry.KeyHandles[0]);
+		if (KeyIndex == AllKeyCount - 1)
+		{
+			bIsLastOnlySelectedKey = true;
+		}
+	}
+
+	if (!Entry.CanEaseCurveKey(InKey))
+	{
+		return false;
+	}
+
+	const FName ChannelTypeName = Entry.ChannelHandle.GetChannelTypeName();
+	if (ChannelTypeName == FMovieSceneDoubleChannel::StaticStruct()->GetFName())
+	{
+		if (!FEaseCurveKeyCache::HasWeightedBrokenTangents<FMovieSceneDoubleChannel, FMovieSceneDoubleValue>(InKey, ChannelHandle))
+		{
+			return false;
+		}
+	}
+	else if (ChannelTypeName == FMovieSceneFloatChannel::StaticStruct()->GetFName())
+	{
+		if (!FEaseCurveKeyCache::HasWeightedBrokenTangents<FMovieSceneFloatChannel, FMovieSceneFloatValue>(InKey, ChannelHandle))
+		{
+			return false;
+		}
+	}
+
+	return true;
+}
+
+bool FEaseCurveKeySelection::ProcessCurveEditorKeySelection(FCurveModel* const CurveModel, const FKeyHandle InKey)
+{
+	if (!CurveModel || InKey == FKeyHandle::Invalid())
+	{
+		return false;
+	}
+
+	FMovieSceneChannelHandle ChannelHandle;
+	bool bIsEaseCurve = true;
+
+	if (const FChannelCurveModel<FMovieSceneFloatChannel, float, float>* FloatChannelModel = 
+		static_cast<const FChannelCurveModel<FMovieSceneFloatChannel, float, float>*>(CurveModel))
+	{
+		ChannelHandle = FloatChannelModel->GetChannelHandle();
+		if (!FEaseCurveKeyCache::HasWeightedBrokenTangents<FMovieSceneDoubleChannel, FMovieSceneDoubleValue>(InKey, ChannelHandle))
+		{
+			bIsEaseCurve = false;
+		}
+	}
+	else if (const FChannelCurveModel<FMovieSceneDoubleChannel, double, double>* DoubleChannelModel = 
+		static_cast<const FChannelCurveModel<FMovieSceneDoubleChannel, double, double>*>(CurveModel))
+	{
+		ChannelHandle = DoubleChannelModel->GetChannelHandle();
+		if (!FEaseCurveKeyCache::HasWeightedBrokenTangents<FMovieSceneDoubleChannel, FMovieSceneDoubleValue>(InKey, ChannelHandle))
+		{
+			bIsEaseCurve = false;
+		}
+	}
+
+	if (!ChannelHandle.Get())
+	{
+		return false;
+	}
+
+	const FCurveModelID CurveId = CurveModel->GetId().GetValue();
+
+	FEaseCurveKeyCache& Entry = CurveEditorKeyData.FindOrAdd(CurveId);
+	Entry.ObjectToModify.Emplace<FCurveModel*>(CurveModel);
+	Entry.ChannelHandle = ChannelHandle;
+	Entry.KeyHandles.Add(InKey);
+
+	TotalSelectedKeys++;
+
+	const int32 AllKeyCount = ChannelHandle.Get()->GetNumKeys();
+	const int32 SelectedKeyCount = Entry.KeyHandles.Num();
+
+	if (SelectedKeyCount == 1 && TotalSelectedKeys == 1)
+	{
+		const int32 KeyIndex = ChannelHandle.Get()->GetIndex(Entry.KeyHandles[0]);
+		if (KeyIndex == AllKeyCount - 1)
+		{
+			bIsLastOnlySelectedKey = true;
+		}
+	}
+
+	return bIsEaseCurve;
+}
+
+void FEaseCurveKeySelection::ForEachEaseableKey(const bool bInIncludeEqualValueKeys
+	, const TFunctionRef<bool(const FKeyHandle& /*InKeyHandle*/
+		, const FKeyHandle& /*InNextKeyHandle*/
+		, const FMovieSceneChannelHandle& /*InChannelHandle*/
+		, const TVariant<TObjectPtr<UMovieSceneSection>, FCurveModel*>&)>& InCallable)
+{
+	if (!SequencerKeyData.IsEmpty())
+	{
+		for (const TPair<FName, FEaseCurveKeyCache>& ChannelEntry : SequencerKeyData)
+		{
+			const FName ChannelTypeName = ChannelEntry.Value.ChannelHandle.GetChannelTypeName();
+			if (ChannelTypeName == FMovieSceneDoubleChannel::StaticStruct()->GetFName())
+			{
+				if (!CheckMatchingValues<FMovieSceneDoubleChannel, FMovieSceneDoubleValue>(bInIncludeEqualValueKeys
+					, ChannelEntry.Value.ChannelHandle
+					, ChannelEntry.Value.KeyHandles
+					, ChannelEntry.Value.ObjectToModify
+					, InCallable))
+				{
+					return;
+				}
+			}
+			else if (ChannelTypeName == FMovieSceneFloatChannel::StaticStruct()->GetFName())
+			{
+				if (!CheckMatchingValues<FMovieSceneFloatChannel, FMovieSceneFloatValue>(bInIncludeEqualValueKeys
+					, ChannelEntry.Value.ChannelHandle
+					, ChannelEntry.Value.KeyHandles
+					, ChannelEntry.Value.ObjectToModify
+					, InCallable))
+				{
+					return;
+				}
+			}
+		}
+	}
+
+	if (!CurveEditorKeyData.IsEmpty())
+	{
+		for (const TPair<FCurveModelID, FEaseCurveKeyCache>& Entry : CurveEditorKeyData)
+		{
+			const FName ChannelTypeName = Entry.Value.ChannelHandle.GetChannelTypeName();
+			if (ChannelTypeName == FMovieSceneDoubleChannel::StaticStruct()->GetFName())
+			{
+				if (!CheckMatchingValues<FMovieSceneDoubleChannel, FMovieSceneDoubleValue>(bInIncludeEqualValueKeys
+					, Entry.Value.ChannelHandle
+					, Entry.Value.KeyHandles
+					, Entry.Value.ObjectToModify
+					, InCallable))
+				{
+					return;
+				}
+			}
+			else if (ChannelTypeName == FMovieSceneFloatChannel::StaticStruct()->GetFName())
+			{
+				if (!CheckMatchingValues<FMovieSceneFloatChannel, FMovieSceneFloatValue>(bInIncludeEqualValueKeys
+					, Entry.Value.ChannelHandle
+					, Entry.Value.KeyHandles
+					, Entry.Value.ObjectToModify
+					, InCallable))
+				{
+					return;
+				}
+			}
+		}
+	}
+}
+
+FEaseCurveTangents FEaseCurveKeySelection::AverageTangents(const FFrameRate& InDisplayRate
+	, const FFrameRate& InTickResolution
+	, const bool bInAutoFlipTangents
+	, const bool bInFirstOnly)
+{
+	TArray<FEaseCurveTangents> KeySetTangents;
+	TArray<FEaseCurveTangents> ChangingTangents;
+
+	ForEachEaseableKey(/*bInIncludeEqualValueKeys=*/true
+		, [&InDisplayRate, &InTickResolution, bInAutoFlipTangents, bInFirstOnly, &KeySetTangents, &ChangingTangents]
+			(const FKeyHandle& InKeyHandle
+			, const FKeyHandle& InNextKeyHandle
+			, const FMovieSceneChannelHandle& InChannelHandle
+			, const TVariant<TObjectPtr<UMovieSceneSection>, FCurveModel*>& InObjectToModify)
+		{
+			const FName ChannelTypeName = InChannelHandle.GetChannelTypeName();
+			if (ChannelTypeName == FMovieSceneDoubleChannel::StaticStruct()->GetFName())
+			{
+				NormalizeChannelValues<FMovieSceneDoubleChannel, FMovieSceneDoubleValue>(InKeyHandle, InNextKeyHandle, InChannelHandle
+					, bInAutoFlipTangents, InDisplayRate, InTickResolution, KeySetTangents, ChangingTangents);
+			}
+			else if (ChannelTypeName == FMovieSceneFloatChannel::StaticStruct()->GetFName())
+			{
+				NormalizeChannelValues<FMovieSceneFloatChannel, FMovieSceneFloatValue>(InKeyHandle, InNextKeyHandle, InChannelHandle
+					, bInAutoFlipTangents, InDisplayRate, InTickResolution, KeySetTangents, ChangingTangents);
+			}
+
+			return !bInFirstOnly;
+		});
+
+	return FEaseCurveTangents::Average(ChangingTangents);
+}
+
+void FEaseCurveKeySelection::SetTangents(const FEaseCurveTangents& InTangents
+	, const EEaseCurveToolOperation InOperation
+	, const FFrameRate& InDisplayRate
+	, const FFrameRate& InTickResolution
+	, const bool bInAutoFlipTangents)
+{
+	if (TotalSelectedKeys == 0)
+	{
+		return;
+	}
+
+	ForEachEaseableKey(/*bInIncludeEqualValueKeys=*/true, [&InDisplayRate, &InTickResolution, &InTangents, InOperation, bInAutoFlipTangents]
+			(const FKeyHandle& InKeyHandle
+			, const FKeyHandle& InNextKeyHandle
+			, const FMovieSceneChannelHandle& InChannelHandle
+			, const TVariant<TObjectPtr<UMovieSceneSection>, FCurveModel*>& InObjectToModify)
+		{
+			const FName ChannelTypeName = InChannelHandle.GetChannelTypeName();
+			if (ChannelTypeName == FMovieSceneDoubleChannel::StaticStruct()->GetFName())
+			{
+				SetChannelValues<FMovieSceneDoubleChannel, FMovieSceneDoubleValue>(InTangents
+					, InOperation, InKeyHandle, InNextKeyHandle, InChannelHandle
+					, bInAutoFlipTangents, InDisplayRate, InTickResolution, InObjectToModify);
+			}
+			else if (ChannelTypeName == FMovieSceneFloatChannel::StaticStruct()->GetFName())
+			{
+				SetChannelValues<FMovieSceneFloatChannel, FMovieSceneFloatValue>(InTangents
+					, InOperation, InKeyHandle, InNextKeyHandle, InChannelHandle
+					, bInAutoFlipTangents, InDisplayRate, InTickResolution, InObjectToModify);
+			}
+
+			return true;
+		});
+}
+
+int32 FEaseCurveKeySelection::GetKeyChannelCount() const
+{
+	const int32 CurveEditorKeyCount = CurveEditorKeyData.Num();
+	if (CurveEditorKeyCount > 0)
+	{
+		return CurveEditorKeyCount;
+	}
+
+	const int32 SequencerKeyCount = SequencerKeyData.Num();
+	if (SequencerKeyCount > 0)
+	{
+		return SequencerKeyCount;
+	}
+
+	return 0;
+}
+
+} // namespace UE::EaseCurveTool
diff --git a/Engine/Plugins/Experimental/EaseCurveTool/Source/EaseCurveTool/Private/EaseCurveKeySelection.h b/Engine/Plugins/Experimental/EaseCurveTool/Source/EaseCurveTool/Private/EaseCurveKeySelection.h
new file mode 100644
--- /dev/null
+++ b/Engine/Plugins/Experimental/EaseCurveTool/Source/EaseCurveTool/Private/EaseCurveKeySelection.h
@@ -0,0 +1,165 @@
+﻿// Copyright Epic Games, Inc. All Rights Reserved.
+
+#pragma once
+
+#include "CurveEditorTypes.h"
+#include "Channels/MovieSceneChannelData.h"
+#include "Channels/MovieSceneChannelHandle.h"
+#include "EaseCurveKeyCache.h"
+
+class FCurveEditor;
+class FCurveModel;
+class ISequencer;
+class UMovieSceneSection;
+struct FCurveEditorSelection;
+struct FEaseCurveTangents;
+struct FKeyHandle;
+
+namespace UE::Sequencer
+{
+	class FChannelModel;
+	class FSequencerSelection;
+}
+
+namespace UE::EaseCurveTool
+{
+
+enum class EEaseCurveToolOperation : uint8;
+
+struct FEaseCurveKeySelection
+{
+	/** Constructs an empty ease curve key selection */
+	FEaseCurveKeySelection() {}
+	/** Constructs an ease curve key selection from a sequencer key selection */
+	FEaseCurveKeySelection(const TSharedPtr<ISequencer>& InSequencer);
+
+	void ForEachEaseableKey(const bool bInIncludeEqualValueKeys
+		, const TFunctionRef<bool(const FKeyHandle& /*InKeyHandle*/
+		, const FKeyHandle& /*InNextKeyHandle*/
+		, const FMovieSceneChannelHandle& /*InChannelHandle*/
+		, const TVariant<TObjectPtr<UMovieSceneSection>, FCurveModel*>&)>& InCallable);
+
+	template <class InChannelType, class InValueType>
+	bool CheckMatchingValues(const bool bInIncludeEqualValueKeys
+		, const FMovieSceneChannelHandle& InChannelHandle
+		, const TArray<FKeyHandle>& InKeyHandles
+		, const TVariant<TObjectPtr<UMovieSceneSection>, FCurveModel*>& InObjectToModify
+		, const TFunctionRef<bool(const FKeyHandle& /*InKeyHandle*/
+			, const FKeyHandle& /*InNextKeyHandle*/
+			, const FMovieSceneChannelHandle& /*InChannelHandle*/
+			, const TVariant<TObjectPtr<UMovieSceneSection>, FCurveModel*>&)>& InCallable)
+	{
+		if (!InChannelHandle.Get())
+		{
+			return false;
+		}
+		
+		TMovieSceneChannelHandle<InChannelType> Channel = InChannelHandle.Cast<InChannelType>();
+		TMovieSceneChannelData<InValueType> ChannelData = Channel.Get()->GetData();
+
+		const TArrayView<InValueType> ChannelValues = ChannelData.GetValues();
+		const int32 KeyCount = ChannelValues.Num();
+
+		for (const FKeyHandle& KeyHandle : InKeyHandles)
+		{
+			if (KeyHandle == FKeyHandle::Invalid())
+			{
+				continue;
+			}
+
+			const int32 KeyIndex = ChannelData.GetIndex(KeyHandle);
+			if (KeyIndex == INDEX_NONE)
+			{
+				continue;
+			}
+
+			// If there is no key after the selected key, we don't need to process.
+			// The arrive tangents of this key will be set by the previous key's processing.
+			int32 NextKeyIndex = KeyIndex + 1;
+			NextKeyIndex = (NextKeyIndex < KeyCount) ? NextKeyIndex : INDEX_NONE;
+			if (NextKeyIndex == INDEX_NONE)
+			{
+				continue;
+			}
+
+			// Need to check if the next key index is valid, otherwise GetHandle() will fail.
+			const FKeyHandle NextKeyHandle = ChannelData.GetHandle(NextKeyIndex);
+			if (NextKeyHandle == FKeyHandle::Invalid())
+			{
+				continue;
+			}
+
+			if (!bInIncludeEqualValueKeys && ChannelValues[KeyIndex].Value == ChannelValues[NextKeyIndex].Value)
+			{
+				continue;
+			}
+
+			if (!InCallable(KeyHandle, NextKeyHandle, InChannelHandle, InObjectToModify))
+			{
+				return false;
+			}
+		}
+
+		return true;
+	}
+
+	FEaseCurveTangents AverageTangents(const FFrameRate& InDisplayRate
+		, const FFrameRate& InTickResolution
+		, const bool bInAutoFlipTangents
+		, const bool bInFirstOnly = false);
+
+	void SetTangents(const FEaseCurveTangents& InTangents
+		, const EEaseCurveToolOperation InOperation
+		, const FFrameRate& InDisplayRate
+		, const FFrameRate& InTickResolution
+		, const bool bInAutoFlipTangents);
+
+	const TMap<FName, FEaseCurveKeyCache>& GetSequencerChannelKeyData() const
+	{
+		return SequencerKeyData;
+	}
+
+	const TMap<FCurveModelID, FEaseCurveKeyCache>& GetCurveEditorChannelKeyData() const
+	{
+		return CurveEditorKeyData;
+	}
+
+	int32 GetTotalSelectedKeys() const
+	{
+		return TotalSelectedKeys;
+	}
+
+	/** @return True if there is only one selected key and it is the last key of the channel. */
+	bool IsLastOnlySelectedKey() const
+	{
+		return bIsLastOnlySelectedKey;
+	}
+
+	/** @return True if all selected keys are detected to be ease curves (broken, weighted, cubic tangents). */
+	bool AreAllEaseCurves() const
+	{
+		return bAreAllEaseCurves;
+	}
+
+	bool IsCurveEditorSelection() const
+	{
+		return !CurveEditorKeyData.IsEmpty();
+	}
+
+	int32 GetKeyChannelCount() const;
+
+protected:
+	bool ProcessSequencerKeySelection(const TSharedPtr<Sequencer::FSequencerSelection>& InSequencerSelection, const FKeyHandle InKey);
+	bool ProcessCurveEditorKeySelection(FCurveModel* const CurveModel, const FKeyHandle InKey);
+
+	TMap<FName, FEaseCurveKeyCache> SequencerKeyData;
+	TMap<FCurveModelID, FEaseCurveKeyCache> CurveEditorKeyData;
+
+	int32 TotalSelectedKeys = 0;
+
+	bool bIsLastOnlySelectedKey = false;
+
+	bool bAreAllEaseCurves = true;
+};
+
+} // namespace UE::EaseCurveTool
diff --git a/Engine/Plugins/Experimental/EaseCurveTool/Source/EaseCurveTool/Private/EaseCurveToolExtender.cpp b/Engine/Plugins/Experimental/EaseCurveTool/Source/EaseCurveTool/Private/EaseCurveToolExtender.cpp
new file mode 100644
--- /dev/null
+++ b/Engine/Plugins/Experimental/EaseCurveTool/Source/EaseCurveTool/Private/EaseCurveToolExtender.cpp
@@ -0,0 +1,351 @@
+﻿// Copyright Epic Games, Inc. All Rights Reserved.
+
+#include "EaseCurveToolExtender.h"
+#include "EaseCurveTool.h"
+#include "EaseCurveToolCommands.h"
+#include "EaseCurveToolSettings.h"
+#include "EaseCurveToolSidebarSection.h"
+#include "EaseCurveToolUtils.h"
+#include "Features/IModularFeatures.h"
+#include "Framework/MultiBox/MultiBoxBuilder.h"
+#include "ISequencer.h"
+#include "ISequencerModule.h"
+#include "Modules/ModuleManager.h"
+#include "SequencerSettings.h"
+#include "Toolkits/AssetEditorToolkit.h"
+#include "ToolMenus.h"
+#include "ToolMenu.h"
+#include "ToolMenuSection.h"
+#include "Widgets/Layout/SExpandableArea.h"
+#include "Widgets/SNullWidget.h"
+#include "Widgets/Text/STextBlock.h"
+
+#define LOCTEXT_NAMESPACE "EaseCurveToolExtender"
+
+namespace UE::EaseCurveTool
+{
+
+FEaseCurveToolExtender& FEaseCurveToolExtender::Get()
+{
+	static FEaseCurveToolExtender Instance;
+	return Instance;
+}
+
+FEaseCurveToolExtender::FEaseCurveToolExtender()
+{
+	IModularFeatures::Get().RegisterModularFeature(GetModularFeatureName(), this);
+
+	ISequencerModule& SequencerModule = FModuleManager::LoadModuleChecked<ISequencerModule>(TEXT("Sequencer"));
+	SequencerCreatedHandle = SequencerModule.RegisterOnSequencerCreated(
+		FOnSequencerCreated::FDelegate::CreateRaw(this, &FEaseCurveToolExtender::OnSequencerCreated));
+}
+
+FEaseCurveToolExtender::~FEaseCurveToolExtender()
+{
+	IModularFeatures::Get().UnregisterModularFeature(GetModularFeatureName(), this);
+
+	if (FModuleManager::Get().IsModuleLoaded(TEXT("Sequencer")))
+	{
+		ISequencerModule& SequencerModule = FModuleManager::LoadModuleChecked<ISequencerModule>(TEXT("Sequencer"));
+		SequencerModule.UnregisterOnSequencerCreated(SequencerCreatedHandle);
+	}
+
+	for (const TPair<FName, FEaseCurveToolInstance>& ToolInstance : ToolInstances)
+	{
+		//RemoveSequencerToolBarExtension(ToolInstance.Value);
+	}
+}
+
+FEaseCurveToolExtender::FEaseCurveToolInstance* FEaseCurveToolExtender::FindOrAddToolInstance_Internal(
+	const TSharedRef<ISequencer>& InSequencer)
+{
+	const FName NewToolId = GetToolInstanceId(*InSequencer);
+	if (NewToolId.IsNone())
+	{
+		return nullptr;
+	}
+
+	FEaseCurveToolInstance& NewToolInstance = ToolInstances.FindOrAdd(NewToolId);
+
+	NewToolInstance.ToolId = NewToolId;
+	NewToolInstance.Instance = MakeShared<FEaseCurveTool>(InSequencer);
+	NewToolInstance.WeakSequencer = InSequencer;
+
+	NewToolInstance.SequencerClosedHandle = InSequencer->OnCloseEvent().AddLambda(
+		[this, NewToolId](const TSharedRef<ISequencer> InSequencer)
+		{
+			OnSequencerClosed(NewToolId, InSequencer);
+		});
+
+	InSequencer->RegisterDrawerSection(TEXT("SelectionDetails"), MakeShared<FEaseCurveToolSidebarSection>(InSequencer));
+
+	//AddSidebarExtension(NewToolInstance);
+
+	return &NewToolInstance;
+}
+
+void FEaseCurveToolExtender::OnSequencerCreated(const TSharedRef<ISequencer> InSequencer)
+{
+	FEaseCurveToolInstance* const NewToolInstance = FindOrAddToolInstance_Internal(InSequencer);
+	if (!NewToolInstance || !NewToolInstance->Instance.IsValid())
+	{
+		return;
+	}
+}
+
+void FEaseCurveToolExtender::OnSequencerClosed(const FName InToolId, const TSharedRef<ISequencer> InSequencer)
+{
+	const FName ToolId = GetToolInstanceId(InSequencer);
+
+	if (ToolId.IsNone() || !ToolInstances.Contains(ToolId))
+	{
+		return;
+	}
+
+	if (!ToolInstances[ToolId].WeakSequencer.IsValid()
+		|| ToolInstances[ToolId].WeakSequencer.Pin().ToSharedRef() != InSequencer)
+	{
+		return;
+	}
+
+	if (ToolInstances[ToolId].SidebarExtender.IsValid())
+	{
+		ISequencerModule& SequencerModule = FModuleManager::LoadModuleChecked<ISequencerModule>(TEXT("Sequencer"));
+		SequencerModule.GetSidebarExtensibilityManager()->RemoveExtender(ToolInstances[ToolId].SidebarExtender);
+		ToolInstances[ToolId].SidebarExtender.Reset();
+	}
+
+	InSequencer->UnregisterDrawerSection(TEXT("SelectionDetails"), FEaseCurveToolSidebarSection::UniqueId);
+
+	RemoveSidebarExtension(ToolInstances[ToolId]);
+
+	ToolInstances.Remove(ToolId);
+}
+
+FName FEaseCurveToolExtender::GetModularFeatureName()
+{
+	static FName FeatureName = TEXT("NavigationTool");
+	return FeatureName;
+}
+
+FName FEaseCurveToolExtender::GetToolInstanceId(const ISequencer& InSequencer)
+{
+	const USequencerSettings* const SequencerSettings = InSequencer.GetSequencerSettings();
+	return SequencerSettings ? SequencerSettings->GetFName() : NAME_None;
+}
+
+FName FEaseCurveToolExtender::GetToolInstanceId(const FEaseCurveTool& InTool)
+{
+	if (const TSharedPtr<ISequencer> Sequencer = InTool.GetSequencer())
+	{
+		return GetToolInstanceId(*Sequencer);
+	}
+	return NAME_None;
+}
+
+TSharedPtr<FEaseCurveTool> FEaseCurveToolExtender::GetToolInstance(const FName InToolId)
+{
+	if (FEaseCurveToolInstance* const ToolInstance = Get().ToolInstances.Find(InToolId))
+	{
+		return ToolInstance->Instance;
+	}
+	return nullptr;
+}
+
+TSharedPtr<FEaseCurveTool> FEaseCurveToolExtender::FindToolInstance(const TSharedRef<ISequencer>& InSequencer)
+{
+	const FName ToolId = GetToolInstanceId(*InSequencer);
+	if (!ToolId.IsNone())
+	{
+		return GetToolInstance(ToolId);
+	}
+	return nullptr;
+}
+
+TSharedPtr<FEaseCurveTool> FEaseCurveToolExtender::FindToolInstanceByCurveEditor(const TSharedRef<FCurveEditor>& InCurveEditor)
+{
+	for (const TPair<FName, FEaseCurveToolInstance>& ToolInstance : Get().ToolInstances)
+	{
+		if (const TSharedPtr<ISequencer> Sequencer = ToolInstance.Value.WeakSequencer.Pin())
+		{
+			if (const TSharedPtr<FCurveEditor> CurveEditor = GetCurveEditorFromSequencer(Sequencer.ToSharedRef()))
+			{
+				if (CurveEditor == InCurveEditor)
+				{
+					return ToolInstance.Value.Instance;
+				}
+			}
+		}
+	}
+	return nullptr;
+}
+
+void FEaseCurveToolExtender::AddSidebarExtension(FEaseCurveToolInstance& InToolInstance)
+{
+	check(InToolInstance.Instance.IsValid());
+
+	UToolMenu* const SidebarMenu = UToolMenus::Get()->ExtendMenu(TEXT("Sequencer.Sidebar"));
+	SidebarMenu->AddDynamicSection(NAME_None, FNewToolMenuDelegate::CreateStatic(&FEaseCurveToolExtender::CreateSidebarWidget, InToolInstance.ToolId));
+
+	/*InToolInstance.SidebarExtender = MakeShared<FExtender>();
+
+	InToolInstance.SidebarExtender->AddMenuExtension(TEXT("KeyEdit")
+		, EExtensionHook::First
+		, InToolInstance.Instance->GetCommandList()
+		, FMenuExtensionDelegate::CreateStatic(&FEaseCurveToolExtender::CreateSidebarWidget, InToolInstance.ToolId)
+	);
+
+	ISequencerModule& SequencerModule = FModuleManager::LoadModuleChecked<ISequencerModule>(TEXT("Sequencer"));
+	SequencerModule.GetSidebarExtensibilityManager()->AddExtender(InToolInstance.SidebarExtender);*/
+}
+
+void FEaseCurveToolExtender::RemoveSidebarExtension(FEaseCurveToolInstance& InToolInstance)
+{
+	if (InToolInstance.SidebarExtender.IsValid())
+	{
+		ISequencerModule& SequencerModule = FModuleManager::LoadModuleChecked<ISequencerModule>(TEXT("Sequencer"));
+		SequencerModule.GetSidebarExtensibilityManager()->RemoveExtender(InToolInstance.SidebarExtender);
+	}
+}
+
+void FEaseCurveToolExtender::CreateSidebarWidget(UToolMenu* const InMenu, const FName InToolId)
+{
+	const TSharedPtr<FEaseCurveTool> ToolInstance = GetToolInstance(InToolId);
+	if (!ToolInstance.IsValid())
+	{
+		return;
+	}
+
+	const TSharedRef<SExpandableArea> ToolMenuWidget =
+		SNew(SExpandableArea)
+		.AreaTitle(LOCTEXT("EaseCurveToolTitle", "Ease Curve Tool"))
+		.InitiallyCollapsed(false)
+		.BorderImage(FAppStyle::Get().GetBrush("Brushes.Header"))
+		.BodyBorderImage(FAppStyle::Get().GetBrush("Brushes.Recessed"))
+		.HeaderPadding(FMargin(2.f, 5.f))
+		.BodyContent()
+		[
+			ToolInstance->GenerateWidget()
+		];
+
+	/*FMenuEntryStyleParams StyleParams;
+	StyleParams.bNoIndent = true;
+	StyleParams.HorizontalAlignment = HAlign_Fill;
+	StyleParams.VerticalAlignment = VAlign_Top;*/
+
+	FToolMenuEntry ToolMenuEntry = FToolMenuEntry::InitWidget(TEXT("EaseCurveTool"), ToolMenuWidget, FText::GetEmpty());
+	ToolMenuEntry.Visibility = TAttribute<EVisibility>::CreateLambda([InToolId]()
+		{
+			if (const TSharedPtr<FEaseCurveTool> ToolInstance = GetToolInstance(InToolId))
+			{
+				if (const UEaseCurveToolSettings* const ToolSettings = GetDefault<UEaseCurveToolSettings>())
+				{
+					return ToolSettings->ShouldShowInSidebar() && !ToolInstance->IsToolTabVisible()
+						? EVisibility::Visible : EVisibility::Collapsed;
+				}
+			}
+			return EVisibility::Collapsed;
+		});
+	/*MenuBuilder.AddWidget(ToolMenuWidget
+		, FText::GetEmpty()
+		, StyleParams
+		, true
+		, FText::GetEmpty()
+		, FSlateIcon()
+		, TAttribute<EVisibility>::CreateLambda([InToolId]()
+			{
+				if (const TSharedPtr<FEaseCurveTool> ToolInstance = GetToolInstance(InToolId))
+				{
+					if (const UEaseCurveToolSettings* const ToolSettings = GetDefault<UEaseCurveToolSettings>())
+					{
+						return ToolSettings->ShouldShowInSidebar() && !ToolInstance->IsToolTabVisible()
+							? EVisibility::Visible : EVisibility::Collapsed;
+					}
+				}
+				return EVisibility::Collapsed;
+			})
+	);*/
+}
+
+void FEaseCurveToolExtender::CreateSidebarWidget(FMenuBuilder& MenuBuilder, const FName InToolId)
+{
+	const TSharedPtr<FEaseCurveTool> ToolInstance = GetToolInstance(InToolId);
+	if (!ToolInstance.IsValid())
+	{
+		return;
+	}
+
+	const TSharedRef<SExpandableArea> ToolMenuWidget =
+		SNew(SExpandableArea)
+		.AreaTitle(LOCTEXT("EaseCurveToolTitle", "Ease Curve Tool"))
+		.InitiallyCollapsed(false)
+		.BorderImage(FAppStyle::Get().GetBrush("Brushes.Header"))
+		.BodyBorderImage(FAppStyle::Get().GetBrush("Brushes.Recessed"))
+		.HeaderPadding(FMargin(2.f, 5.f))
+		.BodyContent()
+		[
+			ToolInstance->GenerateWidget()
+		];
+
+	FMenuEntryStyleParams StyleParams;
+	StyleParams.bNoIndent = true;
+	StyleParams.HorizontalAlignment = HAlign_Fill;
+	StyleParams.VerticalAlignment = VAlign_Top;
+
+	MenuBuilder.AddWidget(ToolMenuWidget
+		, FText::GetEmpty()
+		, StyleParams
+		, /*bInSearchable=*/true
+		, FText::GetEmpty()
+		, FSlateIcon()
+		, TAttribute<EVisibility>::CreateLambda([InToolId]()
+			{
+				if (const TSharedPtr<FEaseCurveTool> ToolInstance = GetToolInstance(InToolId))
+				{
+					if (const UEaseCurveToolSettings* const ToolSettings = GetDefault<UEaseCurveToolSettings>())
+					{
+						return ToolSettings->ShouldShowInSidebar() && !ToolInstance->IsToolTabVisible()
+							? EVisibility::Visible : EVisibility::Collapsed;
+					}
+				}
+				return EVisibility::Collapsed;
+			})
+	);
+}
+
+TSharedRef<SWidget> FEaseCurveToolExtender::MakeQuickPresetMenu(const FName InToolId)
+{
+	UToolMenus* ToolMenus = UToolMenus::Get();
+	if (!ToolMenus)
+	{
+		return SNullWidget::NullWidget;
+	}
+
+	FEaseCurveToolInstance* const ToolInstance = Get().ToolInstances.Find(InToolId);
+	if (!ToolInstance || !ToolInstance->Instance.IsValid())
+	{
+		return SNullWidget::NullWidget;
+	}
+
+	constexpr const TCHAR* MenuName = TEXT("EaseCurveTool.QuickPreset");
+
+	if (!ToolMenus->IsMenuRegistered(MenuName))
+	{
+		UToolMenu* const ToolMenu = ToolMenus->RegisterMenu(MenuName, NAME_None, EMultiBoxType::Menu);
+
+		const FEaseCurveToolCommands& EaseCurveToolCommands = FEaseCurveToolCommands::Get();
+
+		FToolMenuSection& EaseCurveToolSection = ToolMenu->FindOrAddSection(TEXT("EaseCurveToolActions")
+			, LOCTEXT("EaseCurveToolActions", "Ease Curve Tool"));
+
+		EaseCurveToolSection.AddMenuEntry(EaseCurveToolCommands.QuickEase);
+		EaseCurveToolSection.AddMenuEntry(EaseCurveToolCommands.QuickEaseIn);
+		EaseCurveToolSection.AddMenuEntry(EaseCurveToolCommands.QuickEaseOut);
+	}
+
+	return ToolMenus->GenerateWidget(MenuName, FToolMenuContext(ToolInstance->Instance->GetCommandList()));
+}
+
+} // namespace UE::EaseCurveTool
+
+#undef LOCTEXT_NAMESPACE
diff --git a/Engine/Plugins/Experimental/EaseCurveTool/Source/EaseCurveTool/Private/EaseCurveToolModule.cpp b/Engine/Plugins/Experimental/EaseCurveTool/Source/EaseCurveTool/Private/EaseCurveToolModule.cpp
new file mode 100644
--- /dev/null
+++ b/Engine/Plugins/Experimental/EaseCurveTool/Source/EaseCurveTool/Private/EaseCurveToolModule.cpp
@@ -0,0 +1,89 @@
+﻿// Copyright Epic Games, Inc. All Rights Reserved.
+
+#include "EaseCurveToolModule.h"
+#include "ContentBrowserModule.h"
+#include "EaseCurveEditorExtension.h"
+#include "EaseCurveStyle.h"
+#include "EaseCurveToolCommands.h"
+#include "EaseCurveToolExtender.h"
+#include "Framework/MultiBox/MultiBoxExtender.h"
+#include "ICurveEditorModule.h"
+#include "Menus/EaseCurveLibraryMenu.h"
+
+void FEaseCurveToolModule::StartupModule()
+{
+	using namespace UE::EaseCurveTool;
+
+	FEaseCurveToolCommands::Register();
+
+	// Ensure singleton instances are created
+	FEaseCurveStyle::Get();
+	FEaseCurveToolExtender::Get();
+
+	RegisterContentBrowserExtender();
+	RegisterCurveEditorExtender();
+}
+
+void FEaseCurveToolModule::ShutdownModule()
+{
+	UnregisterCurveEditorExtender();
+	UnregisterContentBrowserExtender();
+
+	FEaseCurveToolCommands::Unregister();
+}
+
+void FEaseCurveToolModule::RegisterContentBrowserExtender()
+{
+	FContentBrowserModule& ContentBrowserModule = FModuleManager::LoadModuleChecked<FContentBrowserModule>(TEXT("ContentBrowser"));
+	TArray<FContentBrowserMenuExtender_SelectedAssets>& AssetMenuExtenders = ContentBrowserModule.GetAllAssetViewContextMenuExtenders();
+
+	AssetMenuExtenders.Add(FContentBrowserMenuExtender_SelectedAssets::CreateStatic(&UE::EaseCurveTool::FEaseCurveLibraryMenu::OnExtendContentBrowserAssetSelectionMenu));
+	ContentBrowserExtenderHandle = AssetMenuExtenders.Last().GetHandle();
+}
+
+void FEaseCurveToolModule::UnregisterContentBrowserExtender()
+{
+	if (!ContentBrowserExtenderHandle.IsValid()
+		|| !FModuleManager::Get().IsModuleLoaded(TEXT("ContentBrowser")))
+	{
+		return;
+	}
+
+	FContentBrowserModule& ContentBrowserModule = FModuleManager::LoadModuleChecked<FContentBrowserModule>(TEXT("ContentBrowser"));
+	TArray<FContentBrowserMenuExtender_SelectedAssets>& AssetMenuExtenders = ContentBrowserModule.GetAllAssetViewContextMenuExtenders();
+
+	AssetMenuExtenders.RemoveAll([this](const FContentBrowserMenuExtender_SelectedAssets& InDelegate)
+		{
+			return InDelegate.GetHandle() == ContentBrowserExtenderHandle;
+		});
+	ContentBrowserExtenderHandle.Reset();
+}
+
+void FEaseCurveToolModule::RegisterCurveEditorExtender()
+{
+	using namespace UE::EaseCurveTool;
+
+	ICurveEditorModule& CurveEditorModule = FModuleManager::LoadModuleChecked<ICurveEditorModule>(TEXT("CurveEditor"));
+
+	CurveEditorExtenderHandle = CurveEditorModule.RegisterEditorExtension(
+		FOnCreateCurveEditorExtension::CreateLambda([](TWeakPtr<FCurveEditor> InCurveEditor)
+		{
+			return MakeShared<FEaseCurveEditorExtension>(InCurveEditor);
+		}));
+}
+
+void FEaseCurveToolModule::UnregisterCurveEditorExtender()
+{
+	if (!CurveEditorExtenderHandle.IsValid()
+		|| !FModuleManager::Get().IsModuleLoaded(TEXT("CurveEditor")))
+	{
+		return;
+	}
+
+	ICurveEditorModule& CurveEditorModule = FModuleManager::LoadModuleChecked<ICurveEditorModule>(TEXT("CurveEditor"));
+
+	CurveEditorModule.UnregisterEditorExtension(CurveEditorExtenderHandle);
+	CurveEditorExtenderHandle.Reset();
+}
+
+IMPLEMENT_MODULE(FEaseCurveToolModule, EaseCurveTool)
diff --git a/Engine/Plugins/Experimental/EaseCurveTool/Source/EaseCurveTool/Private/EaseCurveToolUtils.h b/Engine/Plugins/Experimental/EaseCurveTool/Source/EaseCurveTool/Private/EaseCurveToolUtils.h
new file mode 100644
--- /dev/null
+++ b/Engine/Plugins/Experimental/EaseCurveTool/Source/EaseCurveTool/Private/EaseCurveToolUtils.h
@@ -0,0 +1,24 @@
+﻿// Copyright Epic Games, Inc. All Rights Reserved.
+
+#pragma once
+
+#include "CurveEditor.h"
+#include "EaseCurveTool.h"
+#include "ISequencer.h"
+
+class UEaseCurveLibrary;
+struct FCurveEditorSelection;
+
+namespace UE::EaseCurveTool
+{
+
+class FEaseCurveTool;
+
+UEaseCurveLibrary* GetToolPresetLibrary(const TWeakPtr<FEaseCurveTool>& InWeakTool);
+
+/** Returns True if the curve editor selections match */
+bool CompareCurveEditorSelections(const FCurveEditorSelection& InSelectionA, const FCurveEditorSelection& InSelectionB);
+
+TSharedPtr<FCurveEditor> GetCurveEditorFromSequencer(const TSharedRef<ISequencer>& InSequencer);
+
+} // namespace UE::EaseCurveTool
diff --git a/Engine/Plugins/Experimental/EaseCurveTool/Source/EaseCurveTool/Private/Menus/EaseCurveLibraryMenu.cpp b/Engine/Plugins/Experimental/EaseCurveTool/Source/EaseCurveTool/Private/Menus/EaseCurveLibraryMenu.cpp
new file mode 100644
--- /dev/null
+++ b/Engine/Plugins/Experimental/EaseCurveTool/Source/EaseCurveTool/Private/Menus/EaseCurveLibraryMenu.cpp
@@ -0,0 +1,221 @@
+﻿// Copyright Epic Games, Inc. All Rights Reserved.
+
+#include "EaseCurveLibraryMenu.h"
+#include "AssetRegistry/AssetData.h"
+#include "EaseCurveLibrary.h"
+#include "EaseCurveSerializer.h"
+#include "Framework/MultiBox/MultiBoxBuilder.h"
+#include "Framework/MultiBox/MultiBoxExtender.h"
+#include "UObject/UObjectIterator.h"
+
+#define LOCTEXT_NAMESPACE "EaseCurveLibraryMenu"
+
+namespace UE::EaseCurveTool
+{
+
+TSharedRef<FExtender> FEaseCurveLibraryMenu::OnExtendContentBrowserAssetSelectionMenu(const TArray<FAssetData>& InSelectedAssets)
+{
+	using namespace UE::EaseCurveTool;
+
+	const TSharedRef<FExtender> Extender = MakeShared<FExtender>();
+
+	TSet<TWeakObjectPtr<UEaseCurveLibrary>> WeakLibraryAssets;
+
+	for (const FAssetData& SelectedAsset : InSelectedAssets)
+	{
+		if (UEaseCurveLibrary* const LibraryAsset = Cast<UEaseCurveLibrary>(SelectedAsset.GetAsset()))
+		{
+			WeakLibraryAssets.Add(LibraryAsset);
+		}
+	}
+
+	if (WeakLibraryAssets.IsEmpty())
+	{
+		return Extender;
+	}
+
+	Extender->AddMenuExtension(TEXT("GetAssetActions"),
+		EExtensionHook::After,
+		nullptr,
+		FMenuExtensionDelegate::CreateStatic(&FEaseCurveLibraryMenu::AddMenuImportExportSection, WeakLibraryAssets)
+	);
+
+	return Extender;
+}
+
+void FEaseCurveLibraryMenu::AddMenuEntryForNoSerializers(FMenuBuilder& InMenuBuilder)
+{
+	InMenuBuilder.AddMenuEntry(LOCTEXT("NoSerializersName", "No serializers found")
+		, LOCTEXT("NoSerializersTooltip", "To enable import/export options, a custom serializer should be setup")
+		, FSlateIcon()
+		, FExecuteAction()
+	);
+}
+
+bool FEaseCurveLibraryMenu::AddMenuEntryForSerializer(FMenuBuilder& InMenuBuilder
+	, UEaseCurveSerializer& InSerializer
+	, const bool bInImport
+	, const TSet<TWeakObjectPtr<UEaseCurveLibrary>> InWeakLibraryAssets)
+{
+	if (InSerializer.HasAnyFlags(RF_ClassDefaultObject)
+		&& !InSerializer.GetClass()->HasAnyClassFlags(CLASS_Deprecated | CLASS_Abstract))
+	{
+		const FName IconBrushName = bInImport ? TEXT("Icons.Toolbar.Import") : TEXT("Icons.Toolbar.Export");
+
+		InMenuBuilder.AddMenuEntry(InSerializer.GetDisplayName(),
+			InSerializer.GetDisplayTooltip(),
+			FSlateIcon(FAppStyle::GetAppStyleSetName(), IconBrushName),
+			FExecuteAction::CreateStatic(&FEaseCurveLibraryMenu::PromptForExport, &InSerializer, InWeakLibraryAssets)
+		);
+
+		return true;
+	}
+
+	return false;
+}
+
+void FEaseCurveLibraryMenu::AddMenuImportExportSection(FMenuBuilder& InMenuBuilder
+	, const TSet<TWeakObjectPtr<UEaseCurveLibrary>> InWeakLibraryAssets)
+{
+	InMenuBuilder.BeginSection(TEXT("EaseCurveLibrary"), LOCTEXT("EaseCurveLibrary", "Ease Curve Library"));
+
+	InMenuBuilder.AddSubMenu(LOCTEXT("Import", "Import")
+		, LOCTEXT("ImportTooltip", "Import presets to this library from external file sources")
+		, FNewMenuDelegate::CreateStatic(&FEaseCurveLibraryMenu::PopulateImportMenu, InWeakLibraryAssets)
+		, FUIAction()
+		, NAME_None
+		, EUserInterfaceActionType::Button
+		, false
+		, FSlateIcon(FAppStyle::GetAppStyleSetName(), TEXT("Icons.Toolbar.Import")));
+
+	InMenuBuilder.AddSubMenu(LOCTEXT("Export", "Export")
+		, LOCTEXT("ExportTooltip", "Export presets from this library to external file sources")
+		, FNewMenuDelegate::CreateStatic(&FEaseCurveLibraryMenu::PopulateExportMenu, InWeakLibraryAssets)
+		, FUIAction()
+		, NAME_None
+		, EUserInterfaceActionType::Button
+		, false
+		, FSlateIcon(FAppStyle::GetAppStyleSetName(), TEXT("Icons.Toolbar.Export")));
+
+	InMenuBuilder.AddMenuEntry(LOCTEXT("SetToDefaultPresets", "Set to Default Presets")
+		, LOCTEXT("SetToDefaultPresetsTooltip", "Removes all current presets and adds all default presets")
+		, FSlateIcon(FAppStyle::GetAppStyleSetName(), TEXT("Icons.Toolbar.Import"))
+		, FExecuteAction::CreateStatic(&FEaseCurveLibraryMenu::SetToDefaultPresets, InWeakLibraryAssets));
+
+	InMenuBuilder.EndSection();
+}
+
+void FEaseCurveLibraryMenu::PopulateImportMenu(FMenuBuilder& InMenuBuilder
+	, const TSet<TWeakObjectPtr<UEaseCurveLibrary>> InWeakLibraryAssets)
+{
+	InMenuBuilder.BeginSection(TEXT("EaseCurveLibrary_Import"), LOCTEXT("ImportSection", "Import"));
+
+	bool bIsEmpty = true;
+
+	for (TObjectIterator<UEaseCurveSerializer> SerializerIt(RF_NoFlags); SerializerIt; ++SerializerIt)
+	{
+		if (UEaseCurveSerializer* const Serializer = *SerializerIt)
+		{
+			if (AddMenuEntryForSerializer(InMenuBuilder, *Serializer, /*bInImport=*/true, InWeakLibraryAssets))
+			{
+				bIsEmpty = false;
+			}
+		}
+	}
+
+	// Add placeholder text to help the user understand they need to create a custom serializer to enable options in this menu
+	if (bIsEmpty)
+	{
+		AddMenuEntryForNoSerializers(InMenuBuilder);
+	}
+
+	InMenuBuilder.EndSection();
+}
+
+void FEaseCurveLibraryMenu::PopulateExportMenu(FMenuBuilder& InMenuBuilder
+	, const TSet<TWeakObjectPtr<UEaseCurveLibrary>> InWeakLibraryAssets)
+{
+	InMenuBuilder.BeginSection(TEXT("EaseCurveLibrary_Export"), LOCTEXT("ExportSection", "Export"));
+
+	bool bIsEmpty = true;
+
+	for (TObjectIterator<UEaseCurveSerializer> SerializerIt(RF_NoFlags); SerializerIt; ++SerializerIt)
+	{
+		if (UEaseCurveSerializer* const Serializer = *SerializerIt)
+		{
+			if (AddMenuEntryForSerializer(InMenuBuilder, *Serializer, /*bInImport=*/false, InWeakLibraryAssets))
+			{
+				bIsEmpty = false;
+			}
+		}
+	}
+
+	// Add placeholder text to help the user understand they need to create a custom serializer to enable options in this menu
+	if (bIsEmpty)
+	{
+		AddMenuEntryForNoSerializers(InMenuBuilder);
+	}
+
+	InMenuBuilder.EndSection();
+}
+
+void FEaseCurveLibraryMenu::PromptForImport(UEaseCurveSerializer* const InSerializer
+	, const TSet<TWeakObjectPtr<UEaseCurveLibrary>> InWeakLibraryAssets)
+{
+	FString FilePath;
+	if (!InSerializer || !InSerializer->PromptUserForFilePath(FilePath, /*bInImport=*/true))
+	{
+		return;
+	}
+
+	TSet<UEaseCurveLibrary*> LibraryAssets;
+
+	for (const TWeakObjectPtr<UEaseCurveLibrary>& WeakLibraryAsset : InWeakLibraryAssets)
+	{
+		if (WeakLibraryAsset.IsValid())
+		{
+			LibraryAssets.Add(WeakLibraryAsset.Get());
+		}
+	}
+
+	InSerializer->Import(FilePath, LibraryAssets);
+}
+
+void FEaseCurveLibraryMenu::PromptForExport(UEaseCurveSerializer* const InSerializer
+	, const TSet<TWeakObjectPtr<UEaseCurveLibrary>> InWeakLibraryAssets)
+{
+	FString FilePath;
+	if (!InSerializer || !InSerializer->PromptUserForFilePath(FilePath, /*bInImport=*/false))
+	{
+		return;
+	}
+
+	TSet<UEaseCurveLibrary*> LibraryAssets;
+
+	for (const TWeakObjectPtr<UEaseCurveLibrary>& WeakLibraryAsset : InWeakLibraryAssets)
+	{
+		if (WeakLibraryAsset.IsValid())
+		{
+			LibraryAssets.Add(WeakLibraryAsset.Get());
+		}
+	}
+
+	InSerializer->Export(FilePath, LibraryAssets);
+}
+
+void FEaseCurveLibraryMenu::SetToDefaultPresets(const TSet<TWeakObjectPtr<UEaseCurveLibrary>> InWeakLibraryAssets)
+{
+	TSet<UEaseCurveLibrary*> LibraryAssets;
+
+	for (const TWeakObjectPtr<UEaseCurveLibrary>& WeakLibraryAsset : InWeakLibraryAssets)
+	{
+		if (WeakLibraryAsset.IsValid())
+		{
+			WeakLibraryAsset->SetToDefaultPresets();
+		}
+	}
+}
+
+} // namespace UE::EaseCurveTool
+
+#undef LOCTEXT_NAMESPACE
diff --git a/Engine/Plugins/Experimental/EaseCurveTool/Source/EaseCurveTool/Public/EaseCurveSerializer.h b/Engine/Plugins/Experimental/EaseCurveTool/Source/EaseCurveTool/Public/EaseCurveSerializer.h
new file mode 100644
--- /dev/null
+++ b/Engine/Plugins/Experimental/EaseCurveTool/Source/EaseCurveTool/Public/EaseCurveSerializer.h
@@ -0,0 +1,115 @@
+﻿// Copyright Epic Games, Inc. All Rights Reserved.
+
+#pragma once
+
+#include "DesktopPlatformModule.h"
+#include "Framework/Application/SlateApplication.h"
+#include "IDesktopPlatform.h"
+#include "Internationalization/Text.h"
+#include "Misc/Paths.h"
+#include "EaseCurveSerializer.generated.h"
+
+class UEaseCurveLibrary;
+class FString;
+class FText;
+
+#define UE_API EASECURVETOOL_API
+
+/**
+ * Derive from this class to make additional ease curve serializers capable of custom
+ * saving and loading of presets.
+ */
+UCLASS(MinimalAPI, Abstract)
+class UEaseCurveSerializer : public UObject
+{
+	GENERATED_BODY()
+
+public:
+	static bool PromptUserForFilePath(FString& OutFilePath, const bool bInImport)
+	{
+		IDesktopPlatform* const DesktopPlatform = FDesktopPlatformModule::Get();
+		if (!DesktopPlatform)
+		{
+			return false;
+		}
+
+		const void* const ParentWindowHandle = FSlateApplication::Get().FindBestParentWindowHandleForDialogs(nullptr);
+		const FString Title = TEXT("Select File Location");
+		const FString FileTypes = TEXT("All Files (*.*)|*.*");
+
+		bool bSuccess;
+		TArray<FString> OutFilePaths;
+
+		if (bInImport)
+		{
+			bSuccess = DesktopPlatform->OpenFileDialog(ParentWindowHandle
+				, Title
+				, FPaths::GetPath(OutFilePath)
+				, FPaths::GetBaseFilename(OutFilePath)
+				, FileTypes
+				, EFileDialogFlags::None
+				, OutFilePaths);
+		}
+		else
+		{
+			bSuccess = DesktopPlatform->SaveFileDialog(ParentWindowHandle
+				, Title
+				, FPaths::GetPath(OutFilePath)
+				, FPaths::GetBaseFilename(OutFilePath)
+				, FileTypes
+				, EFileDialogFlags::None
+				, OutFilePaths);
+		}
+
+		if (!bSuccess || OutFilePaths.Num() != 1)
+		{
+			return false;
+		}
+
+		OutFilePath = OutFilePaths[0];
+
+		return true;
+	}
+
+	/** The display name text to show for menu entries */
+	UE_API virtual FText GetDisplayName() const
+	{
+		return FText::GetEmpty();
+	}
+
+	/** The display tooltip text to show for menu entries */
+	UE_API virtual FText GetDisplayTooltip() const
+	{
+		return FText::GetEmpty();
+	}
+
+	/**
+	 * Exports an Ease Curve Library asset to an external file.
+	 * Override to implement custom export logic (for a specific file type for example)
+	 *
+	 * @param InFilePath The path to the file where the library will be exported
+	 * @param InLibraries The Ease Curve Library assets to export
+	 *
+	 * @return True if the export operation was successful, false otherwise
+	 */
+	UE_API virtual bool Export(const FString& InFilePath, TSet<UEaseCurveLibrary*> InLibraries)
+	{
+		return false;
+	}
+
+	/**
+	 * Imports an Ease Curve Library asset from an external file.
+	 * Override to implement custom import logic (for a specific file type for example)
+	 *
+	 * @param InFilePath The path to the file where the library will be exported
+	 * @param InLibraries The Ease Curve Library assets to import
+	 *
+	 * @return True if the import operation was successful, false otherwise
+	 */
+    UE_API virtual bool Import(const FString& InFilePath, TSet<UEaseCurveLibrary*> InLibraries)
+	{
+		return false;
+	}
+};
+
+#undef UE_API
diff --git a/Engine/Plugins/VirtualProduction/Avalanche/Source/AvalancheSequencer/Private/AvaSequencer.cpp b/Engine/Plugins/VirtualProduction/Avalanche/Source/AvalancheSequencer/Private/AvaSequencer.cpp
--- a/Engine/Plugins/VirtualProduction/Avalanche/Source/AvalancheSequencer/Private/AvaSequencer.cpp
+++ b/Engine/Plugins/VirtualProduction/Avalanche/Source/AvalancheSequencer/Private/AvaSequencer.cpp
@@ -20,9 +20,6 @@
 #include "DetailsView/Section/AvaSequencePlaybackDetails.h"
 #include "DetailsView/Section/AvaSequenceSettingsDetails.h"
 #include "DetailsView/Section/AvaSequenceTreeDetails.h"
-#include "DetailsView/Section/AvaSequencerEaseCurveToolSection.h"
-#include "EaseCurveTool/AvaEaseCurveTool.h"
-#include "EaseCurveTool/AvaEaseCurveToolCommands.h"
 #include "Editor.h"
 #include "Editor/Sequencer/Private/Sequencer.h"
 #include "Framework/Commands/GenericCommands.h"
@@ -222,12 +219,6 @@ FAvaSequencer::FAvaSequencer(IAvaSequencerProvider& InProvider, FAvaSequencerArg
 
 		SidebarExtender = MakeShared<FExtender>();
 
-		SidebarExtender->AddMenuExtension(
-			TEXT("KeyEdit"),
-			EExtensionHook::First,
-			CommandList,
-			FMenuExtensionDelegate::CreateRaw(this, &FAvaSequencer::ExtendSidebarSelectionMenu));
-
 		SidebarExtender->AddMenuExtension(
 			TEXT("MarkedFrames"),
 			EExtensionHook::After,
@@ -275,14 +266,12 @@ FAvaSequencer::~FAvaSequencer()
 
 	if (const TSharedPtr<ISequencer> Sequencer = SequencerWeak.Pin())
 	{
+		Sequencer->UnregisterDrawerSection(SidebarDrawerId, FAvaSequenceTreeDetails::UniqueId);
 		Sequencer->UnregisterDrawerSection(SidebarDrawerId, FAvaSequencePlaybackDetails::UniqueId);
 		Sequencer->UnregisterDrawerSection(SidebarDrawerId, FAvaSequenceSettingsDetails::UniqueId);
-		Sequencer->UnregisterDrawerSection(SidebarDrawerId, FAvaSequencerEaseCurveToolSection::UniqueId);
 
 		Sequencer->UnregisterDrawer(SidebarDrawerId);
 
-		Sequencer->UnregisterDrawerSection(FSequencer::SelectionDrawerId, FAvaSequencerEaseCurveToolSection::UniqueId);
-
 		if (SidebarExtender.IsValid())
 		{
 			ISequencerModule& SequencerModule = FAvaSequencerUtils::GetSequencerModule();
@@ -403,24 +392,6 @@ void FAvaSequencer::EnsureSequencer()
 			SequencerAction->OnSequencerCreated();
 		}
 	}
-
-	// Create ease curve tool and map commands
-	EaseCurveTool = MakeShared<FAvaEaseCurveTool>(Sequencer.ToSharedRef());
-
-	const TSharedRef<FAvaEaseCurveTool> EaseCurveToolRef = EaseCurveTool.ToSharedRef();
-	const FAvaEaseCurveToolCommands& EaseCurveToolCommands = FAvaEaseCurveToolCommands::Get();
-
-	CommandList->MapAction(EaseCurveToolCommands.QuickEaseIn
-		, FExecuteAction::CreateSP(EaseCurveToolRef, &FAvaEaseCurveTool::ApplyQuickEaseToSequencerKeySelections, EAvaEaseCurveToolOperation::In)
-		, FCanExecuteAction::CreateSP(EaseCurveToolRef, &FAvaEaseCurveTool::CanApplyQuickEaseToSequencerKeySelections));
-
-	CommandList->MapAction(EaseCurveToolCommands.QuickEase
-		, FExecuteAction::CreateSP(EaseCurveToolRef, &FAvaEaseCurveTool::ApplyQuickEaseToSequencerKeySelections, EAvaEaseCurveToolOperation::InOut)
-		, FCanExecuteAction::CreateSP(EaseCurveToolRef, &FAvaEaseCurveTool::CanApplyQuickEaseToSequencerKeySelections));
-
-	CommandList->MapAction(EaseCurveToolCommands.QuickEaseOut
-		, FExecuteAction::CreateSP(EaseCurveToolRef, &FAvaEaseCurveTool::ApplyQuickEaseToSequencerKeySelections, EAvaEaseCurveToolOperation::Out)
-		, FCanExecuteAction::CreateSP(EaseCurveToolRef, &FAvaEaseCurveTool::CanApplyQuickEaseToSequencerKeySelections));
 }
 
 TSharedRef<ISequencer> FAvaSequencer::CreateSequencer()
@@ -2104,17 +2075,6 @@ void FAvaSequencer::OnUpdateCameraCut(UObject* InCameraObject, bool bInJumpCut)
 	GetProvider().OnUpdateCameraCut(InCameraObject, bInJumpCut);
 }
 
-TSharedRef<FAvaEaseCurveTool> FAvaSequencer::GetEaseCurveTool() const
-{
-	return EaseCurveTool.ToSharedRef(); 
-}
-
-void FAvaSequencer::ExtendSidebarSelectionMenu(FMenuBuilder& OutMenuBuilder)
-{
-	const TSharedPtr<FAvaSequencerEaseCurveToolSection> Section = MakeShared<FAvaSequencerEaseCurveToolSection>(SharedThis(this));
-	OutMenuBuilder.AddWidget(Section->CreateContentWidget(), FText::GetEmpty(), /*bInNoIndent=*/true);
-}
-
 void FAvaSequencer::ExtendSidebarMarkedFramesMenu(FMenuBuilder& OutMenuBuilder)
 {
 	UAvaSequence* const Sequence = GetViewedSequence();
diff --git a/Engine/Plugins/VirtualProduction/Avalanche/Source/AvalancheSequencer/Private/AvaSequencer.h b/Engine/Plugins/VirtualProduction/Avalanche/Source/AvalancheSequencer/Private/AvaSequencer.h
--- a/Engine/Plugins/VirtualProduction/Avalanche/Source/AvalancheSequencer/Private/AvaSequencer.h
+++ b/Engine/Plugins/VirtualProduction/Avalanche/Source/AvalancheSequencer/Private/AvaSequencer.h
@@ -12,7 +12,6 @@
 #include "Widgets/Views/STreeView.h"
 
 class AActor;
-class FAvaEaseCurveTool;
 class FAvaEditorSelection;
 class FAvaNavigationToolProvider;
 class FAvaSequencerAction;
@@ -169,8 +168,6 @@ public:
 
 	void FixBindingHierarchy();
 
-	TSharedRef<FAvaEaseCurveTool> GetEaseCurveTool() const;
-	
 	TArrayView<TWeakObjectPtr<>> ResolveBoundObjects(const FGuid& InBindingId, UMovieSceneSequence* Sequence) const;
 
 	//~ Begin IAvaSequencer
@@ -245,7 +242,6 @@ private:
 
 	void OnUpdateCameraCut(UObject* InCameraObject, bool bInJumpCut);
 
-	void ExtendSidebarSelectionMenu(FMenuBuilder& OutMenuBuilder);
 	void ExtendSidebarMarkedFramesMenu(FMenuBuilder& OutMenuBuilder);
 
 	void OnSequencerClosed(const TSharedRef<ISequencer> InSequencer);
@@ -310,8 +306,6 @@ private:
 	/** Whether FAvaSequencer is allowed to select to/from the ISequencer instance */
 	const bool bCanProcessSequencerSelections;
 
-	TSharedPtr<FAvaEaseCurveTool> EaseCurveTool;
-
 	TSharedPtr<FExtender> SidebarExtender;
 
 	FDelegateHandle SidebarSelectionExtenderHandle;
diff --git a/Engine/Plugins/VirtualProduction/Avalanche/Source/AvalancheSequencer/Private/AvaSequencerModule.cpp b/Engine/Plugins/VirtualProduction/Avalanche/Source/AvalancheSequencer/Private/AvaSequencerModule.cpp
--- a/Engine/Plugins/VirtualProduction/Avalanche/Source/AvalancheSequencer/Private/AvaSequencerModule.cpp
+++ b/Engine/Plugins/VirtualProduction/Avalanche/Source/AvalancheSequencer/Private/AvaSequencerModule.cpp
@@ -14,9 +14,6 @@
 #include "Customization/AvaSequenceTimeCustomization.h"
 #include "Director/AvaSequenceDirectorBlueprint.h"
 #include "Director/AvaSequenceDirectorCompiler.h"
-#include "EaseCurveTool/AvaEaseCurveStyle.h"
-#include "EaseCurveTool/AvaEaseCurveSubsystem.h"
-#include "EaseCurveTool/AvaEaseCurveToolCommands.h"
 #include "Editor.h"
 #include "IAvaOutliner.h"
 #include "IAvaOutlinerModule.h"
@@ -38,9 +35,6 @@
 void FAvaSequencerModule::StartupModule()
 {
 	FAvaSequencerCommands::Register();
-	FAvaEaseCurveToolCommands::Register();
-
-	FAvaEaseCurveStyle::Get();
 
 	ActorSubobjectSchema = MakeShared<FAvaActorSubobjectSchema>();
 
@@ -55,14 +49,11 @@ void FAvaSequencerModule::StartupModule()
 	RegisterOutlinerItems();
 	RegisterCustomLayouts();
 	RegisterDirectorCompiler();
-
-	EditorInitializedDelegate = FEditorDelegates::OnEditorInitialized.AddRaw(this, &FAvaSequencerModule::OnEditorInitialized);
 }
 
 void FAvaSequencerModule::ShutdownModule()
 {
 	FAvaSequencerCommands::Unregister();
-	FAvaEaseCurveToolCommands::Unregister();
 
 	if (FAvaSequencerUtils::IsSequencerModuleLoaded())
 	{
@@ -79,8 +70,6 @@ void FAvaSequencerModule::ShutdownModule()
 
 	UnregisterOutlinerItems();
 	UnregisterCustomLayouts();
-
-	FEditorDelegates::OnEditorInitialized.Remove(EditorInitializedDelegate);
 }
 
 void FAvaSequencerModule::RegisterSequenceEditor(ISequencerModule& InSequencerModule)
@@ -207,13 +196,6 @@ void FAvaSequencerModule::RegisterDirectorCompiler()
 	);
 }
 
-void FAvaSequencerModule::OnEditorInitialized(const double InDuration)
-{
-	// Copy default ease curve presets from the plugin to the project configured directory
-	// only if the project presets directory is empty.
-	UAvaEaseCurveSubsystem::Get().ResetToDefaultPresets(true);
-}
-
 IMPLEMENT_MODULE(FAvaSequencerModule, AvalancheSequencer)
 
 #undef LOCTEXT_NAMESPACE
diff --git a/Engine/Plugins/VirtualProduction/Avalanche/Source/AvalancheSequencer/Private/AvaSequencerModule.h b/Engine/Plugins/VirtualProduction/Avalanche/Source/AvalancheSequencer/Private/AvaSequencerModule.h
--- a/Engine/Plugins/VirtualProduction/Avalanche/Source/AvalancheSequencer/Private/AvaSequencerModule.h
+++ b/Engine/Plugins/VirtualProduction/Avalanche/Source/AvalancheSequencer/Private/AvaSequencerModule.h
@@ -118,8 +118,6 @@ private:
 
 	void RegisterDirectorCompiler();
 
-	void OnEditorInitialized(const double InDuration);
-
 	TSharedPtr<FAvaActorSubobjectSchema> ActorSubobjectSchema;
 
 	FDelegateHandle SequenceEditorHandle;
@@ -133,6 +131,4 @@ private:
 	FDelegateHandle OutlinerProxiesExtensionDelegateHandle;
 
 	TArray<const UClass*> CustomizedClasses;
-
-	FDelegateHandle EditorInitializedDelegate;
 };
diff --git a/Engine/Plugins/VirtualProduction/Avalanche/Source/AvalancheSequencer/Private/Customization/AvaSequenceCustomization.cpp b/Engine/Plugins/VirtualProduction/Avalanche/Source/AvalancheSequencer/Private/Customization/AvaSequenceCustomization.cpp
--- a/Engine/Plugins/VirtualProduction/Avalanche/Source/AvalancheSequencer/Private/Customization/AvaSequenceCustomization.cpp
+++ b/Engine/Plugins/VirtualProduction/Avalanche/Source/AvalancheSequencer/Private/Customization/AvaSequenceCustomization.cpp
@@ -6,7 +6,6 @@
 #include "AvaTypeSharedPointer.h"
 #include "Commands/AvaSequencerCommands.h"
 #include "DragDropOps/AvaOutlinerItemDragDropOp.h"
-#include "EaseCurveTool/AvaEaseCurveToolCommands.h"
 #include "Framework/MultiBox/MultiBoxBuilder.h"
 #include "Framework/MultiBox/MultiBoxExtender.h"
 #include "ISequencer.h"
@@ -121,17 +120,11 @@ TSharedRef<SWidget> FAvaSequenceCustomization::MakePlaybackMenu() const
 		UToolMenu* const ToolMenu = ToolMenus->RegisterMenu(MenuName, NAME_None, EMultiBoxType::Menu);
 
 		const FAvaSequencerCommands& AvaSequencerCommands = FAvaSequencerCommands::Get();
-		const FAvaEaseCurveToolCommands& EaseCurveToolCommands = FAvaEaseCurveToolCommands::Get();
 
 		FToolMenuSection& SequencerSection = ToolMenu->FindOrAddSection(TEXT("AvaSequencerActions"), LOCTEXT("AvaSequencerActions", "Sequence"));
 		SequencerSection.AddMenuEntry(AvaSequencerCommands.ApplyCurrentState);
 		SequencerSection.AddMenuEntry(AvaSequencerCommands.OpenStaggerTool);
 
-		FToolMenuSection& EaseCurveToolSection = ToolMenu->FindOrAddSection(TEXT("AvaEaseCurveToolActions"), LOCTEXT("AvaEaseCurveToolActions", "Ease Curve Tool"));
-		EaseCurveToolSection.AddMenuEntry(EaseCurveToolCommands.QuickEase);
-		EaseCurveToolSection.AddMenuEntry(EaseCurveToolCommands.QuickEaseIn);
-		EaseCurveToolSection.AddMenuEntry(EaseCurveToolCommands.QuickEaseOut);
-
 		FToolMenuSection& BindingSection = ToolMenu->FindOrAddSection(TEXT("AvaBindingActions"), LOCTEXT("AvaBindingActions", "Binding"));
 		BindingSection.AddMenuEntry(AvaSequencerCommands.FixBindingPaths);
 		BindingSection.AddMenuEntry(AvaSequencerCommands.FixInvalidBindings);
diff --git a/Engine/Plugins/VirtualProduction/Avalanche/Source/AvalancheSequencer/Private/DetailsView/SAvaKeyFrameEdit.cpp b/Engine/Plugins/VirtualProduction/Avalanche/Source/AvalancheSequencer/Private/DetailsView/SAvaKeyFrameEdit.cpp
--- a/Engine/Plugins/VirtualProduction/Avalanche/Source/AvalancheSequencer/Private/DetailsView/SAvaKeyFrameEdit.cpp
+++ b/Engine/Plugins/VirtualProduction/Avalanche/Source/AvalancheSequencer/Private/DetailsView/SAvaKeyFrameEdit.cpp
@@ -2,7 +2,6 @@
 
 #include "SAvaKeyFrameEdit.h"
 #include "AvaSequencer.h"
-#include "EaseCurveTool/AvaEaseCurveTool.h"
 #include "MVVM/Selection/Selection.h"
 #include "MVVM/ViewModels/SequencerEditorViewModel.h"
 #include "Widgets/Layout/SScaleBox.h"
@@ -38,17 +37,6 @@ void SAvaKeyFrameEdit::Construct(const FArguments& InArgs, const TSharedRef<FAva
 			+ SScrollBox::Slot()
 			.AutoSize()
 			.HAlign(HAlign_Fill)
-			.Padding(1.f, 1.f, 1.f, 3.f)
-			[
-				SNew(SScaleBox)
-				.Stretch(EStretch::Fill)
-				[
-					InSequencer->GetEaseCurveTool()->GenerateWidget()
-				]
-			]
-			+ SScrollBox::Slot()
-			.AutoSize()
-			.HAlign(HAlign_Fill)
 			[
 				SNew(SKeyEditInterface, AvaSequencer->GetSequencer())
 				.EditData(KeyEditData)
