[EaseCurveTool] Remove duplicated input chord for SaveCurveEditorSelection command

diff --git a/Engine/Plugins/Experimental/EaseCurveTool/Source/EaseCurveTool/Private/EaseCurveToolCommands.cpp b/Engine/Plugins/Experimental/EaseCurveTool/Source/EaseCurveTool/Private/EaseCurveToolCommands.cpp
--- a/Engine/Plugins/Experimental/EaseCurveTool/Source/EaseCurveTool/Private/EaseCurveToolCommands.cpp
+++ b/Engine/Plugins/Experimental/EaseCurveTool/Source/EaseCurveTool/Private/EaseCurveToolCommands.cpp
@@ -248,10 +248,10 @@ void FEaseCurveToolCommands::RegisterCommands()
 		, SaveCurveEditorSelection
 		, TEXT("SaveCurveEditorSelection")
 		, LOCTEXT("SaveCurveEditorSelection_Label", "Save Curve Editor Selection")
-		, LOCTEXT("SaveCurveEditorSelection_ToolTip", "")
+		, LOCTEXT("SaveCurveEditorSelection_ToolTip", "Saves the current curve editor key selection as a new ease curve preset\n\n*NOTE* The selected key must be using broken and weighted tangents")
 		, FSlateIcon(FAppStyle::Get().GetStyleSetName(), TEXT("Icons.ChevronLeft"))
 		, EUserInterfaceActionType::Button
-		, FInputChord(EKeys::F8, EModifierKey::Command | EModifierKey::Shift));
+		, FInputChord());
 }
 
 #undef LOCTEXT_NAMESPACE
