[EaseCurveTool] Fix a crash when editing keys

#rnx

diff --git a/Engine/Plugins/Experimental/EaseCurveTool/Source/EaseCurveTool/Private/EaseCurveKeyCache.h b/Engine/Plugins/Experimental/EaseCurveTool/Source/EaseCurveTool/Private/EaseCurveKeyCache.h
--- a/Engine/Plugins/Experimental/EaseCurveTool/Source/EaseCurveTool/Private/EaseCurveKeyCache.h
+++ b/Engine/Plugins/Experimental/EaseCurveTool/Source/EaseCurveTool/Private/EaseCurveKeyCache.h
@@ -1,4 +1,4 @@
-﻿// Copyright Epic Games, Inc. All Rights Reserved.
+// Copyright Epic Games, Inc. All Rights Reserved.
 
 #pragma once
 
@@ -28,41 +28,24 @@ public:
 
 	bool CanEaseCurveKey(const FKeyHandle InKeyHandle) const
 	{
-		if (const TObjectPtr<UMovieSceneSection> Section = ObjectToModify.Get<TObjectPtr<UMovieSceneSection>>())
+		if (!ObjectToModify.TryGet<TObjectPtr<UMovieSceneSection>>() && !ObjectToModify.Get<FCurveModel*>())
 		{
-			const FName ChannelTypeName = ChannelHandle.GetChannelTypeName();
-			if (ChannelTypeName == FMovieSceneDoubleChannel::StaticStruct()->GetFName())
-			{
-				if (!HasWeightedBrokenTangents<FMovieSceneDoubleChannel, FMovieSceneDoubleValue>(InKeyHandle, ChannelHandle))
-				{
-					return false;
-				}
-			}
-			else if (ChannelTypeName == FMovieSceneFloatChannel::StaticStruct()->GetFName())
-			{
-				if (!HasWeightedBrokenTangents<FMovieSceneFloatChannel, FMovieSceneFloatValue>(InKeyHandle, ChannelHandle))
-				{
-					return false;
-				}
-			}
+			return false;
 		}
-		else if (const FCurveModel* CurveModel = ObjectToModify.Get<FCurveModel*>())
+		
+		const FName ChannelTypeName = ChannelHandle.GetChannelTypeName();
+		if (ChannelTypeName == FMovieSceneDoubleChannel::StaticStruct()->GetFName())
 		{
-			if (const FChannelCurveModel<FMovieSceneFloatChannel, float, float>* FloatChannelModel = 
-				static_cast<const FChannelCurveModel<FMovieSceneFloatChannel, float, float>*>(CurveModel))
+			if (!HasWeightedBrokenTangents<FMovieSceneDoubleChannel, FMovieSceneDoubleValue>(InKeyHandle, ChannelHandle))
 			{
-				if (!HasWeightedBrokenTangents<FMovieSceneDoubleChannel, FMovieSceneDoubleValue>(InKeyHandle, ChannelHandle))
-				{
-					return false;
-				}
+				return false;
 			}
-			else if (const FChannelCurveModel<FMovieSceneDoubleChannel, double, double>* DoubleChannelModel = 
-				static_cast<const FChannelCurveModel<FMovieSceneDoubleChannel, double, double>*>(CurveModel))
+		}
+		else if (ChannelTypeName == FMovieSceneFloatChannel::StaticStruct()->GetFName())
+		{
+			if (!HasWeightedBrokenTangents<FMovieSceneFloatChannel, FMovieSceneFloatValue>(InKeyHandle, ChannelHandle))
 			{
-				if (!HasWeightedBrokenTangents<FMovieSceneDoubleChannel, FMovieSceneDoubleValue>(InKeyHandle, ChannelHandle))
-				{
-					return false;
-				}
+				return false;
 			}
 		}
 
@@ -86,7 +69,7 @@ public:
 
 		const TArrayView<const InValueType> ChannelValues = Channel->GetValues();
 
-		if (!FEaseCurveTangents::CanEaseCurveKey(ChannelValues[KeyIndex]))
+		if (!ChannelValues.IsValidIndex(KeyIndex) || !FEaseCurveTangents::CanEaseCurveKey(ChannelValues[KeyIndex]))
 		{
 			return false;
 		}
diff --git a/Engine/Plugins/Experimental/EaseCurveTool/Source/EaseCurveTool/Private/EaseCurveKeySelection.cpp b/Engine/Plugins/Experimental/EaseCurveTool/Source/EaseCurveTool/Private/EaseCurveKeySelection.cpp
--- a/Engine/Plugins/Experimental/EaseCurveTool/Source/EaseCurveTool/Private/EaseCurveKeySelection.cpp
+++ b/Engine/Plugins/Experimental/EaseCurveTool/Source/EaseCurveTool/Private/EaseCurveKeySelection.cpp
@@ -161,7 +161,7 @@ FEaseCurveKeySelection::FEaseCurveKeySelection(const TSharedPtr<ISequencer>& InS
 	bAreAllEaseCurves = true;
 
 	bool bCurveEditorSelection = false;
-
+	
 	if (const TSharedPtr<FCurveEditor> CurveEditor = GetCurveEditorFromSequencer(InSequencer.ToSharedRef()))
 	{
 		const FCurveEditorSelection& CurveEditorSelection = CurveEditor->GetSelection();
@@ -172,21 +172,27 @@ FEaseCurveKeySelection::FEaseCurveKeySelection(const TSharedPtr<ISequencer>& InS
 
 			for (const TPair<FCurveModelID, FKeyHandleSet>& SelectedKeyPair : SelectedKeys)
 			{
+				const TUniquePtr<FCurveModel>* CurveModelPtr = Curves.Find(SelectedKeyPair.Key);
+				if (!CurveModelPtr)
+				{
+					continue;
+				}
+
 				for (const FKeyHandle& SelectedKey : SelectedKeyPair.Value.AsArray())
 				{
-					if (Curves.Contains(SelectedKeyPair.Key))
+					const TViewModelPtr<FChannelModel> ChannelModel = SequencerSelection->KeySelection.GetModelForKey(SelectedKey);
+					const TSharedPtr<IKeyArea> KeyArea = ChannelModel.IsValid() ? ChannelModel->GetKeyArea() : nullptr;
+					if (ChannelModel.IsValid() && KeyArea.IsValid())
 					{
-						if (FCurveModel* const CurveModel = Curves[SelectedKeyPair.Key].Get())
+						const FMovieSceneChannelHandle& ChannelHandle = KeyArea->GetChannel();
+						if (!ProcessCurveEditorKeySelection(*CurveModelPtr->Get(), ChannelHandle, SelectedKey))
 						{
-							if (!ProcessCurveEditorKeySelection(CurveModel, SelectedKey))
-							{
-								bAreAllEaseCurves = false;
-								break;
-							}
-							else
-							{
-								bCurveEditorSelection = true;
-							}
+							bAreAllEaseCurves = false;
+							break;
+						}
+						else
+						{
+							bCurveEditorSelection = true;
 						}
 					}
 				}
@@ -275,44 +281,35 @@ bool FEaseCurveKeySelection::ProcessSequencerKeySelection(const TSharedPtr<FSequ
 	return true;
 }
 
-bool FEaseCurveKeySelection::ProcessCurveEditorKeySelection(FCurveModel* const CurveModel, const FKeyHandle InKey)
+bool FEaseCurveKeySelection::ProcessCurveEditorKeySelection(FCurveModel& CurveModel, const FMovieSceneChannelHandle& ChannelHandle, const FKeyHandle& InKey)
 {
-	if (!CurveModel || InKey == FKeyHandle::Invalid())
+	if (!ChannelHandle.Get())
 	{
 		return false;
 	}
 
-	FMovieSceneChannelHandle ChannelHandle;
 	bool bIsEaseCurve = true;
 
-	if (const FChannelCurveModel<FMovieSceneFloatChannel, float, float>* FloatChannelModel = 
-		static_cast<const FChannelCurveModel<FMovieSceneFloatChannel, float, float>*>(CurveModel))
+	const FName ChannelTypeName = ChannelHandle.GetChannelTypeName();
+	if (ChannelTypeName == FMovieSceneDoubleChannel::StaticStruct()->GetFName())
 	{
-		ChannelHandle = FloatChannelModel->GetChannelHandle();
 		if (!FEaseCurveKeyCache::HasWeightedBrokenTangents<FMovieSceneDoubleChannel, FMovieSceneDoubleValue>(InKey, ChannelHandle))
 		{
 			bIsEaseCurve = false;
 		}
 	}
-	else if (const FChannelCurveModel<FMovieSceneDoubleChannel, double, double>* DoubleChannelModel = 
-		static_cast<const FChannelCurveModel<FMovieSceneDoubleChannel, double, double>*>(CurveModel))
+	else if (ChannelTypeName == FMovieSceneFloatChannel::StaticStruct()->GetFName())
 	{
-		ChannelHandle = DoubleChannelModel->GetChannelHandle();
-		if (!FEaseCurveKeyCache::HasWeightedBrokenTangents<FMovieSceneDoubleChannel, FMovieSceneDoubleValue>(InKey, ChannelHandle))
+		if (!FEaseCurveKeyCache::HasWeightedBrokenTangents<FMovieSceneFloatChannel, FMovieSceneFloatValue>(InKey, ChannelHandle))
 		{
 			bIsEaseCurve = false;
 		}
 	}
 
-	if (!ChannelHandle.Get())
-	{
-		return false;
-	}
-
-	const FCurveModelID CurveId = CurveModel->GetId().GetValue();
+	const FCurveModelID CurveId = CurveModel.GetId().GetValue();
 
 	FEaseCurveKeyCache& Entry = CurveEditorKeyData.FindOrAdd(CurveId);
-	Entry.ObjectToModify.Emplace<FCurveModel*>(CurveModel);
+	Entry.ObjectToModify.Emplace<FCurveModel*>(&CurveModel);
 	Entry.ChannelHandle = ChannelHandle;
 	Entry.KeyHandles.Add(InKey);
 
diff --git a/Engine/Plugins/Experimental/EaseCurveTool/Source/EaseCurveTool/Private/EaseCurveKeySelection.h b/Engine/Plugins/Experimental/EaseCurveTool/Source/EaseCurveTool/Private/EaseCurveKeySelection.h
--- a/Engine/Plugins/Experimental/EaseCurveTool/Source/EaseCurveTool/Private/EaseCurveKeySelection.h
+++ b/Engine/Plugins/Experimental/EaseCurveTool/Source/EaseCurveTool/Private/EaseCurveKeySelection.h
@@ -14,6 +14,7 @@ class UMovieSceneSection;
 struct FCurveEditorSelection;
 struct FEaseCurveTangents;
 struct FKeyHandle;
+struct FMovieSceneChannelHandle;
 
 namespace UE::Sequencer
 {
@@ -150,7 +151,7 @@ struct FEaseCurveKeySelection
 
 protected:
 	bool ProcessSequencerKeySelection(const TSharedPtr<Sequencer::FSequencerSelection>& InSequencerSelection, const FKeyHandle InKey);
-	bool ProcessCurveEditorKeySelection(FCurveModel* const CurveModel, const FKeyHandle InKey);
+	bool ProcessCurveEditorKeySelection(FCurveModel& CurveModel, const FMovieSceneChannelHandle& ChannelHandle, const FKeyHandle& InKey);
 
 	TMap<FName, FEaseCurveKeyCache> SequencerKeyData;
 	TMap<FCurveModelID, FEaseCurveKeyCache> CurveEditorKeyData;
