[EaseCurveTool] Fix text localization conflict

diff --git a/Engine/Plugins/Experimental/EaseCurveTool/Source/EaseCurveTool/Private/Widgets/SEaseCurvePresetList.cpp b/Engine/Plugins/Experimental/EaseCurveTool/Source/EaseCurveTool/Private/Widgets/SEaseCurvePresetList.cpp
--- a/Engine/Plugins/Experimental/EaseCurveTool/Source/EaseCurveTool/Private/Widgets/SEaseCurvePresetList.cpp
+++ b/Engine/Plugins/Experimental/EaseCurveTool/Source/EaseCurveTool/Private/Widgets/SEaseCurvePresetList.cpp
@@ -254,7 +254,7 @@ void SEaseCurvePresetList::UpdateGroupsContent()
 	UEaseCurveLibrary* const PresetLibrary = GetToolPresetLibrary(WeakTool);
 	if (!PresetLibrary)
 	{
-		GroupWidgetsParent->SetContent(GenerateNoPresetsWidget(LOCTEXT("NoPresetsLabel", "No library selected")));
+		GroupWidgetsParent->SetContent(GenerateNoPresetsWidget(LOCTEXT("NoLibrarySelected", "No library selected")));
 		return;
 	}
 
