[EaseCurveTool] - Fix regression caused by UE-309971 fix
- Add error enum for handling and display for valid/invalid key selections
- Better graph widget error support

diff --git a/Engine/Plugins/Experimental/EaseCurveTool/Source/EaseCurveTool/Private/EaseCurveEditorExtension.cpp b/Engine/Plugins/Experimental/EaseCurveTool/Source/EaseCurveTool/Private/EaseCurveEditorExtension.cpp
--- a/Engine/Plugins/Experimental/EaseCurveTool/Source/EaseCurveTool/Private/EaseCurveEditorExtension.cpp
+++ b/Engine/Plugins/Experimental/EaseCurveTool/Source/EaseCurveTool/Private/EaseCurveEditorExtension.cpp
@@ -46,24 +46,22 @@ void FEaseCurveEditorExtension::ExtendCurveEditorToolbar(FToolBarBuilder& ToolBa
 
 	ToolBarBuilder.SetLabelVisibility(EVisibility::Collapsed);
 
-	const FText Label = LOCTEXT("ApplyEaseCurveLabel", "Ease Curve Tool");
-	const FText Tooltip = LOCTEXT("ApplyEaseCurveTooltip", "Apply an ease curve between the selected key tangents");
 	const FSlateIcon Icon = FSlateIcon(FEaseCurveStyle::Get().GetStyleSetName(), TEXT("Icon.ToolBar"));
 
 	ToolBarBuilder.AddToolBarButton(FUIAction(
 			FExecuteAction::CreateRaw(this, &FEaseCurveEditorExtension::PopupToolbarMenu),
 			FCanExecuteAction::CreateRaw(this, &FEaseCurveEditorExtension::IsButtonEnabled))
 		, NAME_None
-		, Label
-		, Tooltip
+		, LOCTEXT("EaseCurvePresetsLabel", "Ease Curve Presets")
+		, TAttribute<FText>::CreateRaw(this, &FEaseCurveEditorExtension::GetEaseCurveTooltipText)
 		, Icon);
 
 	ToolBarBuilder.AddComboButton(FUIAction(
 			FExecuteAction(),
 			FCanExecuteAction::CreateRaw(this, &FEaseCurveEditorExtension::IsButtonEnabled))
 		, FOnGetContent::CreateRaw(this, &FEaseCurveEditorExtension::GetToolbarButtonMenuContent)
-		, Label
-		, Tooltip
+		, LOCTEXT("QuickEaseCurveLabel", "Quick Ease Curve")
+		, TAttribute<FText>::CreateRaw(this, &FEaseCurveEditorExtension::GetEaseCurveTooltipText)
 		, Icon
 		, true
 		, NAME_None
@@ -146,11 +144,25 @@ bool FEaseCurveEditorExtension::IsButtonEnabled() const
 {
 	if (const TSharedPtr<FEaseCurveTool> ToolInstance = GetToolInstance())
 	{
-		return ToolInstance->HasCachedKeysToEase();
+		return ToolInstance->HasCachedKeysToEase()
+			&& ToolInstance->GetSelectionError() == EEaseCurveToolError::None;
 	}
 	return false;
 }
 
+FText FEaseCurveEditorExtension::GetEaseCurveTooltipText() const
+{
+	if (const TSharedPtr<FEaseCurveTool> ToolInstance = GetToolInstance())
+	{
+		if (ToolInstance->HasCachedKeysToEase()
+			&& ToolInstance->GetSelectionError() != EEaseCurveToolError::None)
+		{
+			return ToolInstance->GetSelectionErrorText();
+		}
+	}
+	return LOCTEXT("EaseCurvePresetsTooltip", "Apply an ease curve between the selected key tangents");
+}
+
 } // namespace UE::EaseCurveTool
 
 #undef LOCTEXT_NAMESPACE
diff --git a/Engine/Plugins/Experimental/EaseCurveTool/Source/EaseCurveTool/Private/EaseCurveEditorExtension.h b/Engine/Plugins/Experimental/EaseCurveTool/Source/EaseCurveTool/Private/EaseCurveEditorExtension.h
--- a/Engine/Plugins/Experimental/EaseCurveTool/Source/EaseCurveTool/Private/EaseCurveEditorExtension.h
+++ b/Engine/Plugins/Experimental/EaseCurveTool/Source/EaseCurveTool/Private/EaseCurveEditorExtension.h
@@ -42,6 +42,8 @@ private:
 
 	bool IsButtonEnabled() const;
 
+	FText GetEaseCurveTooltipText() const;
+
 	TWeakPtr<FCurveEditor> WeakCurveEditor;
 
 	TSharedPtr<FCurveEditorToolBarMenu> ButtonMenu;
diff --git a/Engine/Plugins/Experimental/EaseCurveTool/Source/EaseCurveTool/Private/EaseCurveKeyCache.h b/Engine/Plugins/Experimental/EaseCurveTool/Source/EaseCurveTool/Private/EaseCurveKeyCache.h
deleted file mode 100644
diff --git a/Engine/Plugins/Experimental/EaseCurveTool/Source/EaseCurveTool/Private/EaseCurveKeySelection.cpp b/Engine/Plugins/Experimental/EaseCurveTool/Source/EaseCurveTool/Private/EaseCurveKeySelection.cpp
--- a/Engine/Plugins/Experimental/EaseCurveTool/Source/EaseCurveTool/Private/EaseCurveKeySelection.cpp
+++ b/Engine/Plugins/Experimental/EaseCurveTool/Source/EaseCurveTool/Private/EaseCurveKeySelection.cpp
@@ -2,6 +2,8 @@
 
 #include "EaseCurveKeySelection.h"
 #include "Channels/ChannelCurveModel.h"
+#include "Channels/DoubleChannelCurveModel.h"
+#include "Channels/FloatChannelCurveModel.h"
 #include "Channels/MovieSceneChannel.h"
 #include "Channels/MovieSceneDoubleChannel.h"
 #include "Channels/MovieSceneFloatChannel.h"
@@ -12,9 +14,9 @@
 #include "IKeyArea.h"
 #include "MVVM/Selection/Selection.h"
 #include "MVVM/ViewModels/ChannelModel.h"
+#include "MVVM/ViewModels/SequenceModel.h"
 #include "MVVM/ViewModels/SequencerEditorViewModel.h"
 #include "ScopedTransaction.h"
-#include "MVVM/ViewModels/SequenceModel.h"
 
 using namespace UE::Sequencer;
 
@@ -70,7 +72,7 @@ void SetChannelValues(const FEaseCurveTangents& InTangents
 	, const FMovieSceneChannelHandle& InChannelHandle
 	, const bool bInAutoFlipTangents
 	, const FFrameRate& InDisplayRate, const FFrameRate& InTickResolution
-	, const TVariant<TObjectPtr<UMovieSceneSection>, FCurveModel*>& InObjectToModify)
+	, const TObjectPtr<UMovieSceneSection>& InSection)
 {
 	TMovieSceneChannelHandle<InChannelType> Channel = /*InChannelData.ChannelHandle*/InChannelHandle.Cast<InChannelType>();
 	TMovieSceneChannelData<InValueType> ChannelData = Channel.Get()->GetData();
@@ -101,21 +103,10 @@ void SetChannelValues(const FEaseCurveTangents& InTangents
 	const FScopedTransaction Transaction(NSLOCTEXT("EaseCurveTool", "SetSequencerCurveTangents", "Set Sequencer Curve Tangents"));
 
 	// Mark the section as changed to force the curve editor to update its display
-	if (InObjectToModify.IsType<TObjectPtr<UMovieSceneSection>>())
+	if (InSection)
 	{
-		if (const TObjectPtr<UMovieSceneSection> SectionToModify = InObjectToModify.Get<TObjectPtr<UMovieSceneSection>>())
-		{
-			SectionToModify->Modify();
-			SectionToModify->MarkAsChanged();
-		}
-	}
-	else if (InObjectToModify.IsType<FCurveModel*>())
-	{
-		if (FCurveModel* const CurveModelToModify = InObjectToModify.Get<FCurveModel*>())
-		{
-			CurveModelToModify->Modify();
-			//CurveModelToModify->MarkAsChanged();
-		}
+		InSection->Modify();
+		InSection->MarkAsChanged();
 	}
 
 	// Set this keys leave tangent
@@ -146,23 +137,7 @@ FEaseCurveKeySelection::FEaseCurveKeySelection(const TSharedPtr<ISequencer>& InS
 		return;
 	}
 
-	const TSharedPtr<FSequencerEditorViewModel> SequencerViewModel = InSequencer->GetViewModel();
-	if (!SequencerViewModel.IsValid())
-	{
-		return;
-	}
-
-	const TSharedPtr<FSequencerSelection> SequencerSelection = SequencerViewModel->GetSelection();
-	if (!SequencerSelection.IsValid())
-	{
-		return;
-	}
-
-	bAreAllEaseCurves = true;
-
-	bool bCurveEditorSelection = false;
-	
-	if (const TSharedPtr<FCurveEditor> CurveEditor = GetCurveEditorFromSequencer(InSequencer.ToSharedRef()))
+	if (const TSharedPtr<FCurveEditor> CurveEditor = FEaseCurveToolUtils::GetCurveEditorFromSequencer(InSequencer.ToSharedRef()))
 	{
 		const FCurveEditorSelection& CurveEditorSelection = CurveEditor->GetSelection();
 		const TMap<FCurveModelID, FKeyHandleSet> SelectedKeys = CurveEditorSelection.GetAll();
@@ -172,45 +147,135 @@ FEaseCurveKeySelection::FEaseCurveKeySelection(const TSharedPtr<ISequencer>& InS
 
 			for (const TPair<FCurveModelID, FKeyHandleSet>& SelectedKeyPair : SelectedKeys)
 			{
-				const TUniquePtr<FCurveModel>* CurveModelPtr = Curves.Find(SelectedKeyPair.Key);
-				if (!CurveModelPtr)
+				const TUniquePtr<FCurveModel>* CurveModelUniquePtr = Curves.Find(SelectedKeyPair.Key);
+				if (!CurveModelUniquePtr)
 				{
 					continue;
 				}
 
-				for (const FKeyHandle& SelectedKey : SelectedKeyPair.Value.AsArray())
+				FCurveModel& CurveModel = *CurveModelUniquePtr->Get();
+
+				bool bContinueProcessing = true;
+
+				FMovieSceneChannelHandle ChannelHandle;
+				if (FEaseCurveToolUtils::FindChannelHandleFromCurveModel(CurveModel, ChannelHandle))
 				{
-					const TViewModelPtr<FChannelModel> ChannelModel = SequencerSelection->KeySelection.GetModelForKey(SelectedKey);
-					const TSharedPtr<IKeyArea> KeyArea = ChannelModel.IsValid() ? ChannelModel->GetKeyArea() : nullptr;
-					if (ChannelModel.IsValid() && KeyArea.IsValid())
+					for (const FKeyHandle& SelectedKey : SelectedKeyPair.Value.AsArray())
 					{
-						const FMovieSceneChannelHandle& ChannelHandle = KeyArea->GetChannel();
-						if (!ProcessCurveEditorKeySelection(*CurveModelPtr->Get(), ChannelHandle, SelectedKey))
+						if (!ProcessCurveEditorKeySelection(CurveModel, ChannelHandle, SelectedKey))
 						{
-							bAreAllEaseCurves = false;
+							bContinueProcessing = false;
 							break;
 						}
-						else
-						{
-							bCurveEditorSelection = true;
-						}
 					}
 				}
+
+				if (!bContinueProcessing)
+				{
+					break;
+				}
 			}
 		}
 	}
 
-	if (!bCurveEditorSelection)
+	// Fallback to sequencer key selection if available
+	if (SelectionError == EEaseCurveToolError::None && CurveEditorKeyData.IsEmpty())
 	{
-		for (const FKeyHandle Key : SequencerSelection->KeySelection)
+		const TSharedPtr<FSequencerEditorViewModel> SequencerViewModel = InSequencer->GetViewModel();
+		if (!SequencerViewModel.IsValid())
+		{
+			return;
+		}
+
+		const TSharedPtr<FSequencerSelection> SequencerSelection = SequencerViewModel->GetSelection();
+		if (!SequencerSelection.IsValid())
+		{
+			return;
+		}
+
+		for (const FKeyHandle& Key : SequencerSelection->KeySelection)
 		{
 			if (!ProcessSequencerKeySelection(SequencerSelection, Key))
 			{
-				bAreAllEaseCurves = false;
 				break;
 			}
 		}
 	}
+
+	SetSelectionError();
+}
+
+void FEaseCurveKeySelection::SetSelectionError()
+{
+	bool bAllLast = true;
+	bool bAllWeightedBroken = true;
+	bool bAllSameValue = true;
+	bool bAnySameValue = false;
+
+	auto LoopKeyHandles = [&bAllLast, &bAllWeightedBroken, &bAllSameValue, &bAnySameValue](const TMap<FKeyHandle, EEaseCurveToolError>& InKeyErrors)
+	{
+		for (const TPair<FKeyHandle, EEaseCurveToolError>& KeyHandlePair : InKeyErrors)
+		{
+			if (EnumHasAnyFlags(KeyHandlePair.Value, EEaseCurveToolError::SameValues))
+			{
+				bAnySameValue = true;
+			}
+			else
+			{
+				bAllSameValue = false;
+			}
+			if (!EnumHasAnyFlags(KeyHandlePair.Value, EEaseCurveToolError::LastKey))
+			{
+				bAllLast = false;
+			}
+			if (EnumHasAnyFlags(KeyHandlePair.Value, EEaseCurveToolError::NoWeightedBrokenCubicTangents))
+			{
+				bAllWeightedBroken = false;
+			}
+		}
+	};
+
+	if (CurveEditorKeyData.Num() > 0)
+	{
+		for (const TPair<FCurveModelID, FEaseCurveChannelKeyCache>& KeyDataPair : CurveEditorKeyData)
+		{
+			LoopKeyHandles(KeyDataPair.Value.KeyHandles);
+		}
+	}
+	else
+	{
+		for (const TPair<FName, FEaseCurveChannelKeyCache>& KeyDataPair : SequencerKeyData)
+		{
+			LoopKeyHandles(KeyDataPair.Value.KeyHandles);
+		}
+	}
+
+	if (bAllSameValue)
+	{
+		if (bAllLast)
+		{
+			SelectionError = EEaseCurveToolError::LastKey;
+		}
+		else
+		{
+			SelectionError = EEaseCurveToolError::SameValues;
+		}
+	}
+	else if (bAllLast)
+	{
+		if (bAnySameValue)
+		{
+			SelectionError = EEaseCurveToolError::SameValues;
+		}
+		else
+		{
+			SelectionError = EEaseCurveToolError::LastKey;
+		}
+	}
+	else if (!bAllWeightedBroken)
+	{
+		SelectionError = EEaseCurveToolError::NoWeightedBrokenCubicTangents;
+	}
 }
 
 bool FEaseCurveKeySelection::ProcessSequencerKeySelection(const TSharedPtr<FSequencerSelection>& InSequencerSelection, const FKeyHandle InKey)
@@ -238,165 +303,100 @@ bool FEaseCurveKeySelection::ProcessSequencerKeySelection(const TSharedPtr<FSequ
 	const FString ChannelName = ChannelModel->GetChannelName().ToString();
 	const FName FullPathMapKey = FName(PathName + ChannelName);
 
-	FEaseCurveKeyCache& Entry = SequencerKeyData.FindOrAdd(FullPathMapKey);
-	Entry.ObjectToModify.Emplace<TObjectPtr<UMovieSceneSection>>(ChannelModel->GetSection());
+	FEaseCurveChannelKeyCache& Entry = SequencerKeyData.FindOrAdd(FullPathMapKey);
+	Entry.Section = ChannelModel->GetSection();
 	Entry.ChannelHandle = ChannelHandle;
-	Entry.KeyHandles.Add(InKey);
+	Entry.KeyHandles.Add(InKey, FindErrorForKey(ChannelHandle, InKey));
 
 	TotalSelectedKeys++;
 
-	const int32 AllKeyCount = ChannelHandle.Get()->GetNumKeys();
-	const int32 SelectedKeyCount = Entry.KeyHandles.Num();
-
-	if (SelectedKeyCount == 1 && TotalSelectedKeys == 1)
-	{
-		const int32 KeyIndex = ChannelHandle.Get()->GetIndex(Entry.KeyHandles[0]);
-		if (KeyIndex == AllKeyCount - 1)
-		{
-			bIsLastOnlySelectedKey = true;
-		}
-	}
+	return true;
+}
 
-	if (!Entry.CanEaseCurveKey(InKey))
+bool FEaseCurveKeySelection::ProcessCurveEditorKeySelection(FCurveModel& InCurveModel
+	, const FMovieSceneChannelHandle& InChannelHandle, const FKeyHandle& InKey)
+{
+	if (!InChannelHandle.Get())
 	{
 		return false;
 	}
 
-	const FName ChannelTypeName = Entry.ChannelHandle.GetChannelTypeName();
-	if (ChannelTypeName == FMovieSceneDoubleChannel::StaticStruct()->GetFName())
-	{
-		if (!FEaseCurveKeyCache::HasWeightedBrokenTangents<FMovieSceneDoubleChannel, FMovieSceneDoubleValue>(InKey, ChannelHandle))
-		{
-			return false;
-		}
-	}
-	else if (ChannelTypeName == FMovieSceneFloatChannel::StaticStruct()->GetFName())
-	{
-		if (!FEaseCurveKeyCache::HasWeightedBrokenTangents<FMovieSceneFloatChannel, FMovieSceneFloatValue>(InKey, ChannelHandle))
-		{
-			return false;
-		}
-	}
+	const FCurveModelID CurveId = InCurveModel.GetId().GetValue();
+
+	FEaseCurveChannelKeyCache& Entry = CurveEditorKeyData.FindOrAdd(CurveId);
+	Entry.Section = InCurveModel.GetOwningObjectOrOuter<UMovieSceneSection>();
+	Entry.ChannelHandle = InChannelHandle;
+	Entry.KeyHandles.Add(InKey, FindErrorForKey(InChannelHandle, InKey));
+
+	TotalSelectedKeys++;
 
 	return true;
 }
 
-bool FEaseCurveKeySelection::ProcessCurveEditorKeySelection(FCurveModel& CurveModel, const FMovieSceneChannelHandle& ChannelHandle, const FKeyHandle& InKey)
+EEaseCurveToolError FEaseCurveKeySelection::FindErrorForKey(const FMovieSceneChannelHandle& InChannelHandle, const FKeyHandle& InKey)
 {
-	if (!ChannelHandle.Get())
-	{
-		return false;
-	}
-
-	bool bIsEaseCurve = true;
+	const FName ChannelTypeName = InChannelHandle.GetChannelTypeName();
 
-	const FName ChannelTypeName = ChannelHandle.GetChannelTypeName();
 	if (ChannelTypeName == FMovieSceneDoubleChannel::StaticStruct()->GetFName())
 	{
-		if (!FEaseCurveKeyCache::HasWeightedBrokenTangents<FMovieSceneDoubleChannel, FMovieSceneDoubleValue>(InKey, ChannelHandle))
-		{
-			bIsEaseCurve = false;
-		}
+		return FindErrorForKey_Internal<FMovieSceneDoubleChannel, FMovieSceneDoubleValue>(InKey, InChannelHandle);
 	}
-	else if (ChannelTypeName == FMovieSceneFloatChannel::StaticStruct()->GetFName())
+	if (ChannelTypeName == FMovieSceneFloatChannel::StaticStruct()->GetFName())
 	{
-		if (!FEaseCurveKeyCache::HasWeightedBrokenTangents<FMovieSceneFloatChannel, FMovieSceneFloatValue>(InKey, ChannelHandle))
-		{
-			bIsEaseCurve = false;
-		}
+		return FindErrorForKey_Internal<FMovieSceneFloatChannel, FMovieSceneFloatValue>(InKey, InChannelHandle);
 	}
 
-	const FCurveModelID CurveId = CurveModel.GetId().GetValue();
-
-	FEaseCurveKeyCache& Entry = CurveEditorKeyData.FindOrAdd(CurveId);
-	Entry.ObjectToModify.Emplace<FCurveModel*>(&CurveModel);
-	Entry.ChannelHandle = ChannelHandle;
-	Entry.KeyHandles.Add(InKey);
-
-	TotalSelectedKeys++;
-
-	const int32 AllKeyCount = ChannelHandle.Get()->GetNumKeys();
-	const int32 SelectedKeyCount = Entry.KeyHandles.Num();
-
-	if (SelectedKeyCount == 1 && TotalSelectedKeys == 1)
-	{
-		const int32 KeyIndex = ChannelHandle.Get()->GetIndex(Entry.KeyHandles[0]);
-		if (KeyIndex == AllKeyCount - 1)
-		{
-			bIsLastOnlySelectedKey = true;
-		}
-	}
-
-	return bIsEaseCurve;
+	return EEaseCurveToolError::None;
 }
 
-void FEaseCurveKeySelection::ForEachEaseableKey(const bool bInIncludeEqualValueKeys
-	, const TFunctionRef<bool(const FKeyHandle& /*InKeyHandle*/
-		, const FKeyHandle& /*InNextKeyHandle*/
-		, const FMovieSceneChannelHandle& /*InChannelHandle*/
-		, const TVariant<TObjectPtr<UMovieSceneSection>, FCurveModel*>&)>& InCallable)
+void FEaseCurveKeySelection::ForEachEaseableKey(const FProcessKeySignature& InCallable)
 {
-	if (!SequencerKeyData.IsEmpty())
+	if (!CurveEditorKeyData.IsEmpty())
 	{
-		for (const TPair<FName, FEaseCurveKeyCache>& ChannelEntry : SequencerKeyData)
+		for (const TPair<FCurveModelID, FEaseCurveChannelKeyCache>& Entry : CurveEditorKeyData)
 		{
-			const FName ChannelTypeName = ChannelEntry.Value.ChannelHandle.GetChannelTypeName();
-			if (ChannelTypeName == FMovieSceneDoubleChannel::StaticStruct()->GetFName())
-			{
-				if (!CheckMatchingValues<FMovieSceneDoubleChannel, FMovieSceneDoubleValue>(bInIncludeEqualValueKeys
-					, ChannelEntry.Value.ChannelHandle
-					, ChannelEntry.Value.KeyHandles
-					, ChannelEntry.Value.ObjectToModify
-					, InCallable))
-				{
-					return;
-				}
-			}
-			else if (ChannelTypeName == FMovieSceneFloatChannel::StaticStruct()->GetFName())
+			if (!ForEachChannelKey(Entry.Value.ChannelHandle
+				, Entry.Value.KeyHandles, Entry.Value.Section, InCallable))
 			{
-				if (!CheckMatchingValues<FMovieSceneFloatChannel, FMovieSceneFloatValue>(bInIncludeEqualValueKeys
-					, ChannelEntry.Value.ChannelHandle
-					, ChannelEntry.Value.KeyHandles
-					, ChannelEntry.Value.ObjectToModify
-					, InCallable))
-				{
-					return;
-				}
+				return;
 			}
 		}
+
+		// Prioritize curve editor key selections or sequencer key selections
+		return;
 	}
 
-	if (!CurveEditorKeyData.IsEmpty())
+	for (const TPair<FName, FEaseCurveChannelKeyCache>& ChannelEntry : SequencerKeyData)
 	{
-		for (const TPair<FCurveModelID, FEaseCurveKeyCache>& Entry : CurveEditorKeyData)
+		if (!ForEachChannelKey(ChannelEntry.Value.ChannelHandle
+			, ChannelEntry.Value.KeyHandles, ChannelEntry.Value.Section, InCallable))
 		{
-			const FName ChannelTypeName = Entry.Value.ChannelHandle.GetChannelTypeName();
-			if (ChannelTypeName == FMovieSceneDoubleChannel::StaticStruct()->GetFName())
-			{
-				if (!CheckMatchingValues<FMovieSceneDoubleChannel, FMovieSceneDoubleValue>(bInIncludeEqualValueKeys
-					, Entry.Value.ChannelHandle
-					, Entry.Value.KeyHandles
-					, Entry.Value.ObjectToModify
-					, InCallable))
-				{
-					return;
-				}
-			}
-			else if (ChannelTypeName == FMovieSceneFloatChannel::StaticStruct()->GetFName())
-			{
-				if (!CheckMatchingValues<FMovieSceneFloatChannel, FMovieSceneFloatValue>(bInIncludeEqualValueKeys
-					, Entry.Value.ChannelHandle
-					, Entry.Value.KeyHandles
-					, Entry.Value.ObjectToModify
-					, InCallable))
-				{
-					return;
-				}
-			}
+			return;
 		}
 	}
 }
 
+bool FEaseCurveKeySelection::ForEachChannelKey(const FMovieSceneChannelHandle& InChannelHandle
+	, const TMap<FKeyHandle, EEaseCurveToolError>& InKeyHandles
+	, const TObjectPtr<UMovieSceneSection>& InSection
+	, const FProcessKeySignature& InCallable)
+{
+	const FName ChannelTypeName = InChannelHandle.GetChannelTypeName();
+
+	if (ChannelTypeName == FMovieSceneDoubleChannel::StaticStruct()->GetFName())
+	{
+		return ForEachChannelKey_Internal<FMovieSceneDoubleChannel, FMovieSceneDoubleValue>(InChannelHandle
+			, InKeyHandles, InSection, InCallable);
+	}
+	if (ChannelTypeName == FMovieSceneFloatChannel::StaticStruct()->GetFName())
+	{
+		return ForEachChannelKey_Internal<FMovieSceneFloatChannel, FMovieSceneFloatValue>(InChannelHandle
+			, InKeyHandles, InSection, InCallable);
+	}
+
+	return true;
+}
+
 FEaseCurveTangents FEaseCurveKeySelection::AverageTangents(const FFrameRate& InDisplayRate
 	, const FFrameRate& InTickResolution
 	, const bool bInAutoFlipTangents
@@ -405,12 +405,11 @@ FEaseCurveTangents FEaseCurveKeySelection::AverageTangents(const FFrameRate& InD
 	TArray<FEaseCurveTangents> KeySetTangents;
 	TArray<FEaseCurveTangents> ChangingTangents;
 
-	ForEachEaseableKey(/*bInIncludeEqualValueKeys=*/true
-		, [&InDisplayRate, &InTickResolution, bInAutoFlipTangents, bInFirstOnly, &KeySetTangents, &ChangingTangents]
+	ForEachEaseableKey([&InDisplayRate, &InTickResolution, bInAutoFlipTangents, bInFirstOnly, &KeySetTangents, &ChangingTangents]
 			(const FKeyHandle& InKeyHandle
 			, const FKeyHandle& InNextKeyHandle
 			, const FMovieSceneChannelHandle& InChannelHandle
-			, const TVariant<TObjectPtr<UMovieSceneSection>, FCurveModel*>& InObjectToModify)
+			, const TObjectPtr<UMovieSceneSection>& InSection)
 		{
 			const FName ChannelTypeName = InChannelHandle.GetChannelTypeName();
 			if (ChannelTypeName == FMovieSceneDoubleChannel::StaticStruct()->GetFName())
@@ -441,24 +440,24 @@ void FEaseCurveKeySelection::SetTangents(const FEaseCurveTangents& InTangents
 		return;
 	}
 
-	ForEachEaseableKey(/*bInIncludeEqualValueKeys=*/true, [&InDisplayRate, &InTickResolution, &InTangents, InOperation, bInAutoFlipTangents]
+	ForEachEaseableKey([&InDisplayRate, &InTickResolution, &InTangents, InOperation, bInAutoFlipTangents]
 			(const FKeyHandle& InKeyHandle
 			, const FKeyHandle& InNextKeyHandle
 			, const FMovieSceneChannelHandle& InChannelHandle
-			, const TVariant<TObjectPtr<UMovieSceneSection>, FCurveModel*>& InObjectToModify)
+			, const TObjectPtr<UMovieSceneSection>& InSection)
 		{
 			const FName ChannelTypeName = InChannelHandle.GetChannelTypeName();
 			if (ChannelTypeName == FMovieSceneDoubleChannel::StaticStruct()->GetFName())
 			{
 				SetChannelValues<FMovieSceneDoubleChannel, FMovieSceneDoubleValue>(InTangents
 					, InOperation, InKeyHandle, InNextKeyHandle, InChannelHandle
-					, bInAutoFlipTangents, InDisplayRate, InTickResolution, InObjectToModify);
+					, bInAutoFlipTangents, InDisplayRate, InTickResolution, InSection);
 			}
 			else if (ChannelTypeName == FMovieSceneFloatChannel::StaticStruct()->GetFName())
 			{
 				SetChannelValues<FMovieSceneFloatChannel, FMovieSceneFloatValue>(InTangents
 					, InOperation, InKeyHandle, InNextKeyHandle, InChannelHandle
-					, bInAutoFlipTangents, InDisplayRate, InTickResolution, InObjectToModify);
+					, bInAutoFlipTangents, InDisplayRate, InTickResolution, InSection);
 			}
 
 			return true;
diff --git a/Engine/Plugins/Experimental/EaseCurveTool/Source/EaseCurveTool/Private/EaseCurveKeySelection.h b/Engine/Plugins/Experimental/EaseCurveTool/Source/EaseCurveTool/Private/EaseCurveKeySelection.h
--- a/Engine/Plugins/Experimental/EaseCurveTool/Source/EaseCurveTool/Private/EaseCurveKeySelection.h
+++ b/Engine/Plugins/Experimental/EaseCurveTool/Source/EaseCurveTool/Private/EaseCurveKeySelection.h
@@ -5,28 +5,56 @@
 #include "CurveEditorTypes.h"
 #include "Channels/MovieSceneChannelData.h"
 #include "Channels/MovieSceneChannelHandle.h"
-#include "EaseCurveKeyCache.h"
+#include "EaseCurveToolUtils.h"
+#include "EaseCurveKeySelection.generated.h"
 
-class FCurveEditor;
 class FCurveModel;
 class ISequencer;
 class UMovieSceneSection;
-struct FCurveEditorSelection;
 struct FEaseCurveTangents;
 struct FKeyHandle;
 struct FMovieSceneChannelHandle;
 
 namespace UE::Sequencer
 {
-	class FChannelModel;
 	class FSequencerSelection;
 }
 
+UENUM()
+enum class EEaseCurveToolError : uint8
+{
+	/** All selected keys are valid ease-able keys */
+	None = 0,
+	/** All selected keys are the last key of the channel and have no next key to ease to */
+	LastKey = 1 << 1,
+	/** All selected keys have the same values (easing would result in no change) */
+	SameValues = 1 << 2,
+	/** All selected keys are not weighted, broken, cubic tangents */
+	NoWeightedBrokenCubicTangents = 1 << 3
+};
+ENUM_CLASS_FLAGS(EEaseCurveToolError)
+
 namespace UE::EaseCurveTool
 {
 
 enum class EEaseCurveToolOperation : uint8;
 
+/** Cached set of key handles for a channel */
+struct FEaseCurveChannelKeyCache
+{
+	TObjectPtr<UMovieSceneSection> Section;
+	FMovieSceneChannelHandle ChannelHandle;
+	TMap<FKeyHandle, EEaseCurveToolError> KeyHandles;
+
+	template <class InChannelHandleType, class InChannelValueType>
+	TMovieSceneChannelData<InChannelValueType> GetChannelData()
+	{
+		const TMovieSceneChannelHandle<InChannelHandleType> Channel = ChannelHandle.Cast<InChannelHandleType>();
+		return Channel.Get()->GetData();
+	}
+};
+
+/** Represents a key selection, either Curve Editor (prioritized) or Sequencer, that can be operated on by the Ease Curve Tool */
 struct FEaseCurveKeySelection
 {
 	/** Constructs an empty ease curve key selection */
@@ -34,41 +62,129 @@ struct FEaseCurveKeySelection
 	/** Constructs an ease curve key selection from a sequencer key selection */
 	FEaseCurveKeySelection(const TSharedPtr<ISequencer>& InSequencer);
 
-	void ForEachEaseableKey(const bool bInIncludeEqualValueKeys
-		, const TFunctionRef<bool(const FKeyHandle& /*InKeyHandle*/
+	FEaseCurveTangents AverageTangents(const FFrameRate& InDisplayRate
+		, const FFrameRate& InTickResolution
+		, const bool bInAutoFlipTangents
+		, const bool bInFirstOnly = false);
+
+	void SetTangents(const FEaseCurveTangents& InTangents
+		, const EEaseCurveToolOperation InOperation
+		, const FFrameRate& InDisplayRate
+		, const FFrameRate& InTickResolution
+		, const bool bInAutoFlipTangents);
+
+	bool HasCachedKeysToEase() const
+	{
+		return !CurveEditorKeyData.IsEmpty() || !SequencerKeyData.IsEmpty();
+	}
+
+	int32 GetTotalSelectedKeys() const
+	{
+		return TotalSelectedKeys;
+	}
+
+	bool IsCurveEditorSelection() const
+	{
+		return !CurveEditorKeyData.IsEmpty();
+	}
+
+	int32 GetKeyChannelCount() const;
+
+	/** @return Error code specifying why the key selection is not valid to be used by the Ease Curve Tool to perform operations on */
+	EEaseCurveToolError GetSelectionError() const
+	{
+		return SelectionError;
+	}
+
+protected:
+	using FProcessKeySignature = TFunctionRef<bool(const FKeyHandle& /*InKeyHandle*/
 		, const FKeyHandle& /*InNextKeyHandle*/
 		, const FMovieSceneChannelHandle& /*InChannelHandle*/
-		, const TVariant<TObjectPtr<UMovieSceneSection>, FCurveModel*>&)>& InCallable);
+		, const TObjectPtr<UMovieSceneSection>& /*InSection*/)>;
+
+	void SetSelectionError();
+
+	EEaseCurveToolError FindErrorForKey(const FMovieSceneChannelHandle& InChannelHandle, const FKeyHandle& InKey);
 
 	template <class InChannelType, class InValueType>
-	bool CheckMatchingValues(const bool bInIncludeEqualValueKeys
-		, const FMovieSceneChannelHandle& InChannelHandle
-		, const TArray<FKeyHandle>& InKeyHandles
-		, const TVariant<TObjectPtr<UMovieSceneSection>, FCurveModel*>& InObjectToModify
-		, const TFunctionRef<bool(const FKeyHandle& /*InKeyHandle*/
-			, const FKeyHandle& /*InNextKeyHandle*/
-			, const FMovieSceneChannelHandle& /*InChannelHandle*/
-			, const TVariant<TObjectPtr<UMovieSceneSection>, FCurveModel*>&)>& InCallable)
+	static EEaseCurveToolError FindErrorForKey_Internal(const FKeyHandle& InKeyHandle, const FMovieSceneChannelHandle& InChannelHandle)
+	{
+		EEaseCurveToolError OutError = EEaseCurveToolError::None;
+
+		InChannelType* const Channel = static_cast<InChannelType*>(InChannelHandle.Get());
+		if (!Channel)
+		{
+			return OutError;
+		}
+
+		const int32 KeyIndex = Channel->GetIndex(InKeyHandle);
+		if (KeyIndex == INDEX_NONE)
+		{
+			return OutError;
+		}
+
+		const TArrayView<const InValueType> ChannelValues = Channel->GetValues();
+		if (!ChannelValues.IsValidIndex(KeyIndex))
+		{
+			return OutError;
+		}
+
+		const int32 NextKeyIndex = KeyIndex + 1;
+		if (ChannelValues.IsValidIndex(NextKeyIndex))
+		{
+			if (ChannelValues[KeyIndex].Value == ChannelValues[NextKeyIndex].Value)
+			{
+				OutError |= EEaseCurveToolError::SameValues;
+			}
+		}
+		else
+		{
+			OutError |= EEaseCurveToolError::LastKey;
+		}
+
+		if (!FEaseCurveToolUtils::HasWeightedBrokenTangents(ChannelValues[KeyIndex]))
+		{
+			OutError |= EEaseCurveToolError::NoWeightedBrokenCubicTangents;
+		}
+
+		return OutError;
+	}
+
+	bool ProcessSequencerKeySelection(const TSharedPtr<Sequencer::FSequencerSelection>& InSequencerSelection, const FKeyHandle InKey);
+	bool ProcessCurveEditorKeySelection(FCurveModel& CurveModel, const FMovieSceneChannelHandle& ChannelHandle, const FKeyHandle& InKey);
+
+	void ForEachEaseableKey(const FProcessKeySignature& InCallable);
+
+	bool ForEachChannelKey(const FMovieSceneChannelHandle& InChannelHandle
+		, const TMap<FKeyHandle, EEaseCurveToolError>& InKeyHandles
+		, const TObjectPtr<UMovieSceneSection>& InSection
+		, const FProcessKeySignature& InCallable);
+
+	template <class InChannelType, class InValueType>
+	bool ForEachChannelKey_Internal(const FMovieSceneChannelHandle& InChannelHandle
+		, const TMap<FKeyHandle, EEaseCurveToolError>& InKeyHandles
+		, const TObjectPtr<UMovieSceneSection>& InSection
+		, const FProcessKeySignature& InCallable)
 	{
 		if (!InChannelHandle.Get())
 		{
 			return false;
 		}
-		
+
 		TMovieSceneChannelHandle<InChannelType> Channel = InChannelHandle.Cast<InChannelType>();
 		TMovieSceneChannelData<InValueType> ChannelData = Channel.Get()->GetData();
 
 		const TArrayView<InValueType> ChannelValues = ChannelData.GetValues();
 		const int32 KeyCount = ChannelValues.Num();
 
-		for (const FKeyHandle& KeyHandle : InKeyHandles)
+		for (const TPair<FKeyHandle, EEaseCurveToolError>& KeyHandlePair : InKeyHandles)
 		{
-			if (KeyHandle == FKeyHandle::Invalid())
+			if (KeyHandlePair.Key == FKeyHandle::Invalid())
 			{
 				continue;
 			}
 
-			const int32 KeyIndex = ChannelData.GetIndex(KeyHandle);
+			const int32 KeyIndex = ChannelData.GetIndex(KeyHandlePair.Key);
 			if (KeyIndex == INDEX_NONE)
 			{
 				continue;
@@ -90,12 +206,7 @@ struct FEaseCurveKeySelection
 				continue;
 			}
 
-			if (!bInIncludeEqualValueKeys && ChannelValues[KeyIndex].Value == ChannelValues[NextKeyIndex].Value)
-			{
-				continue;
-			}
-
-			if (!InCallable(KeyHandle, NextKeyHandle, InChannelHandle, InObjectToModify))
+			if (!InCallable(KeyHandlePair.Key, NextKeyHandle, InChannelHandle, InSection))
 			{
 				return false;
 			}
@@ -104,63 +215,12 @@ struct FEaseCurveKeySelection
 		return true;
 	}
 
-	FEaseCurveTangents AverageTangents(const FFrameRate& InDisplayRate
-		, const FFrameRate& InTickResolution
-		, const bool bInAutoFlipTangents
-		, const bool bInFirstOnly = false);
-
-	void SetTangents(const FEaseCurveTangents& InTangents
-		, const EEaseCurveToolOperation InOperation
-		, const FFrameRate& InDisplayRate
-		, const FFrameRate& InTickResolution
-		, const bool bInAutoFlipTangents);
-
-	const TMap<FName, FEaseCurveKeyCache>& GetSequencerChannelKeyData() const
-	{
-		return SequencerKeyData;
-	}
-
-	const TMap<FCurveModelID, FEaseCurveKeyCache>& GetCurveEditorChannelKeyData() const
-	{
-		return CurveEditorKeyData;
-	}
-
-	int32 GetTotalSelectedKeys() const
-	{
-		return TotalSelectedKeys;
-	}
-
-	/** @return True if there is only one selected key and it is the last key of the channel. */
-	bool IsLastOnlySelectedKey() const
-	{
-		return bIsLastOnlySelectedKey;
-	}
-
-	/** @return True if all selected keys are detected to be ease curves (broken, weighted, cubic tangents). */
-	bool AreAllEaseCurves() const
-	{
-		return bAreAllEaseCurves;
-	}
-
-	bool IsCurveEditorSelection() const
-	{
-		return !CurveEditorKeyData.IsEmpty();
-	}
-
-	int32 GetKeyChannelCount() const;
-
-protected:
-	bool ProcessSequencerKeySelection(const TSharedPtr<Sequencer::FSequencerSelection>& InSequencerSelection, const FKeyHandle InKey);
-	bool ProcessCurveEditorKeySelection(FCurveModel& CurveModel, const FMovieSceneChannelHandle& ChannelHandle, const FKeyHandle& InKey);
-
-	TMap<FName, FEaseCurveKeyCache> SequencerKeyData;
-	TMap<FCurveModelID, FEaseCurveKeyCache> CurveEditorKeyData;
+	TMap<FName, FEaseCurveChannelKeyCache> SequencerKeyData;
+	TMap<FCurveModelID, FEaseCurveChannelKeyCache> CurveEditorKeyData;
 
 	int32 TotalSelectedKeys = 0;
 
-	bool bIsLastOnlySelectedKey = false;
-
-	bool bAreAllEaseCurves = true;
+	EEaseCurveToolError SelectionError = EEaseCurveToolError::None;
 };
 
 } // namespace UE::EaseCurveTool
diff --git a/Engine/Plugins/Experimental/EaseCurveTool/Source/EaseCurveTool/Private/EaseCurveTangents.cpp b/Engine/Plugins/Experimental/EaseCurveTool/Source/EaseCurveTool/Private/EaseCurveTangents.cpp
--- a/Engine/Plugins/Experimental/EaseCurveTool/Source/EaseCurveTool/Private/EaseCurveTangents.cpp
+++ b/Engine/Plugins/Experimental/EaseCurveTool/Source/EaseCurveTool/Private/EaseCurveTangents.cpp
@@ -320,25 +320,4 @@ double FEaseCurveTangents::CalculateCurveLength(const int32 SampleCount) const
 	return Length;
 }
 
-bool FEaseCurveTangents::CanEaseCurveKey(const FRichCurveKey& InKey)
-{
-	return InKey.TangentMode == RCTM_Break
-		&& InKey.TangentWeightMode == RCTWM_WeightedBoth
-		&& InKey.InterpMode == RCIM_Cubic;
-}
-
-bool FEaseCurveTangents::CanEaseCurveKey(const FMovieSceneDoubleValue& InValue)
-{
-	return InValue.TangentMode == RCTM_Break
-		&& InValue.Tangent.TangentWeightMode == RCTWM_WeightedBoth
-		&& InValue.InterpMode == RCIM_Cubic;
-}
-
-bool FEaseCurveTangents::CanEaseCurveKey(const FMovieSceneFloatValue& InValue)
-{
-	return InValue.TangentMode == RCTM_Break
-		&& InValue.Tangent.TangentWeightMode == RCTWM_WeightedBoth
-		&& InValue.InterpMode == RCIM_Cubic;
-}
-
 #undef LOCTEXT_NAMESPACE
diff --git a/Engine/Plugins/Experimental/EaseCurveTool/Source/EaseCurveTool/Private/EaseCurveTangents.h b/Engine/Plugins/Experimental/EaseCurveTool/Source/EaseCurveTool/Private/EaseCurveTangents.h
--- a/Engine/Plugins/Experimental/EaseCurveTool/Source/EaseCurveTool/Private/EaseCurveTangents.h
+++ b/Engine/Plugins/Experimental/EaseCurveTool/Source/EaseCurveTool/Private/EaseCurveTangents.h
@@ -26,10 +26,6 @@ public:
 
 	static FEaseCurveTangents Average(const TArray<FEaseCurveTangents>& InTangentArray);
 
-	static bool CanEaseCurveKey(const FRichCurveKey& InKey);
-	static bool CanEaseCurveKey(const FMovieSceneDoubleValue& InValue);
-	static bool CanEaseCurveKey(const FMovieSceneFloatValue& InValue);
-
 	FEaseCurveTangents() {}
 	FEaseCurveTangents(const double InStart, const double InStartWeight, const double InEnd, const double InEndWeight)
 		: Start(InStart), StartWeight(InStartWeight)
diff --git a/Engine/Plugins/Experimental/EaseCurveTool/Source/EaseCurveTool/Private/EaseCurveTool.cpp b/Engine/Plugins/Experimental/EaseCurveTool/Source/EaseCurveTool/Private/EaseCurveTool.cpp
--- a/Engine/Plugins/Experimental/EaseCurveTool/Source/EaseCurveTool/Private/EaseCurveTool.cpp
+++ b/Engine/Plugins/Experimental/EaseCurveTool/Source/EaseCurveTool/Private/EaseCurveTool.cpp
@@ -6,7 +6,6 @@
 #include "CurveEditor.h"
 #include "CurveEditorSelection.h"
 #include "Dialogs/DlgPickAssetPath.h"
-#include "EaseCurveKeySelection.h"
 #include "EaseCurveLibrary.h"
 #include "EaseCurvePreset.h"
 #include "EaseCurveToolCommands.h"
@@ -44,6 +43,8 @@ FEaseCurveTool::FEaseCurveTool(const TSharedRef<ISequencer>& InSequencer)
 
 	ToolSettings = GetMutableDefault<UEaseCurveToolSettings>();
 
+	WeakPresetLibrary = ToolSettings->GetPresetLibrary();
+
 	UpdateEaseCurveFromKeySelections();
 
 	InSequencer->GetSelectionChangedObjectGuids().AddRaw(this, &FEaseCurveTool::OnSequencerSelectionChanged);
@@ -138,7 +139,6 @@ TSharedRef<SWidget> FEaseCurveTool::GenerateWidget()
 
 		ToolWidget = SNew(SEaseCurveTool, SharedThis(this))
 			.InitialTangents(GetEaseCurveTangents())
-			.Visibility(this, &FEaseCurveTool::GetVisibility)
 			.ToolOperation(this, &FEaseCurveTool::GetToolOperation);
 	}
 
@@ -150,12 +150,6 @@ TSharedPtr<SEaseCurveTool> FEaseCurveTool::GetToolWidget() const
 	return ToolWidget;
 }
 
-EVisibility FEaseCurveTool::GetVisibility() const
-{
-	return (KeyCache.GetTotalSelectedKeys() > 0 && !KeyCache.IsLastOnlySelectedKey())
-		? EVisibility::Visible : EVisibility::Collapsed;
-}
-
 UEaseCurveLibrary* FEaseCurveTool::GetPresetLibrary() const
 {
 	return WeakPresetLibrary.Get();
@@ -168,6 +162,7 @@ void FEaseCurveTool::SetPresetLibrary(const TWeakObjectPtr<UEaseCurveLibrary>& I
 	if (ToolSettings)
 	{
 		ToolSettings->SetPresetLibrary(InPresetLibrary.Get());
+		ToolSettings->SaveConfig();
 	}
 
 	PresetLibraryChangedDelegate.Broadcast(WeakPresetLibrary);
@@ -611,10 +606,10 @@ void FEaseCurveTool::Tick(float InDeltaTime)
 	// of curve editor key selection changes.
 	if (const TSharedPtr<ISequencer> Sequencer = SequencerWeak.Pin())
 	{
-		if (const TSharedPtr<FCurveEditor> CurveEditor = GetCurveEditorFromSequencer(Sequencer.ToSharedRef()))
+		if (const TSharedPtr<FCurveEditor> CurveEditor = FEaseCurveToolUtils::GetCurveEditorFromSequencer(Sequencer.ToSharedRef()))
 		{
 			const FCurveEditorSelection& CurveEditorSelection = CurveEditor->GetSelection();
-			if (!CompareCurveEditorSelections(CurveEditorSelection, LastCurveEditorSelection))
+			if (!FEaseCurveToolUtils::CompareCurveEditorSelections(CurveEditorSelection, LastCurveEditorSelection))
 			{
 				bRefreshForKeySelectionChange = true;
 			}
@@ -679,26 +674,30 @@ void FEaseCurveTool::RecordPresetAnalytics(const TSharedPtr<FEaseCurvePreset>& I
 	}
 }
 
-bool FEaseCurveTool::HasCachedKeysToEase()
+bool FEaseCurveTool::HasCachedKeysToEase() const
 {
-	bool bEaseableKeys = false;
-
-	KeyCache.ForEachEaseableKey(/*bInIncludeEqualValueKeys=*/false, [&bEaseableKeys]
-		(const FKeyHandle& InKeyHandle
-			, const FKeyHandle& InNextKeyHandle
-			, const FMovieSceneChannelHandle& InChannelData
-			, const TVariant<TObjectPtr<UMovieSceneSection>, FCurveModel*>& InObjectToModify)
-		{
-			bEaseableKeys = true;
-			return false;
-		});
+	return KeyCache.HasCachedKeysToEase();
+}
 
-	return bEaseableKeys;
+EEaseCurveToolError FEaseCurveTool::GetSelectionError() const
+{
+	return KeyCache.GetSelectionError();
 }
 
-bool FEaseCurveTool::AreAllEaseCurves()
+FText FEaseCurveTool::GetSelectionErrorText() const
 {
-	return KeyCache.AreAllEaseCurves();
+	switch (GetSelectionError())
+	{
+	case EEaseCurveToolError::None:
+		return FText::GetEmpty();
+	case EEaseCurveToolError::LastKey:
+		return LOCTEXT("LastSelectedKey", "No next key to ease to");
+	case EEaseCurveToolError::SameValues:
+		return LOCTEXT("EqualValueKeys", "No different key values!");
+	case EEaseCurveToolError::NoWeightedBrokenCubicTangents:
+		return LOCTEXT("NoEaseCurve", "No weighted, broken, cubic tangents!");
+	}
+	return FText::GetEmpty();
 }
 
 bool FEaseCurveTool::IsCurveEditorSelection() const
diff --git a/Engine/Plugins/Experimental/EaseCurveTool/Source/EaseCurveTool/Private/EaseCurveTool.h b/Engine/Plugins/Experimental/EaseCurveTool/Source/EaseCurveTool/Private/EaseCurveTool.h
--- a/Engine/Plugins/Experimental/EaseCurveTool/Source/EaseCurveTool/Private/EaseCurveTool.h
+++ b/Engine/Plugins/Experimental/EaseCurveTool/Source/EaseCurveTool/Private/EaseCurveTool.h
@@ -4,7 +4,6 @@
 
 #include "CurveEditorSelection.h"
 #include "Curves/RichCurve.h"
-#include "Curves/KeyHandle.h"
 #include "EaseCurve.h"
 #include "EaseCurveKeySelection.h"
 #include "EditorUndoClient.h"
@@ -12,8 +11,6 @@
 #include "Templates/SharedPointer.h"
 #include "TickableEditorObject.h"
 #include "UObject/WeakObjectPtrTemplates.h"
-#include "Widgets/DeclarativeSyntaxSupport.h"
-#include "Widgets/SCompoundWidget.h"
 #include "Widgets/SEaseCurvePreset.h"
 
 class FUICommandList;
@@ -79,8 +76,6 @@ public:
 	TSharedRef<SWidget> GenerateWidget();
 	TSharedPtr<SEaseCurveTool> GetToolWidget() const;
 
-	EVisibility GetVisibility() const;
-
 	UEaseCurveLibrary* GetPresetLibrary() const;
 	void SetPresetLibrary(const TWeakObjectPtr<UEaseCurveLibrary>& InPresetLibrary);
 
@@ -141,9 +136,10 @@ public:
 	FFrameRate GetTickResolution() const;
 	FFrameRate GetDisplayRate() const;
 
-	bool HasCachedKeysToEase();
+	bool HasCachedKeysToEase() const;
 
-	bool AreAllEaseCurves();
+	EEaseCurveToolError GetSelectionError() const;
+	FText GetSelectionErrorText() const;
 
 	/** @return True if the tool is operating on a curve editor key selection, False if operating on a Sequencer key selection */
 	bool IsCurveEditorSelection() const;
diff --git a/Engine/Plugins/Experimental/EaseCurveTool/Source/EaseCurveTool/Private/EaseCurveToolExtender.cpp b/Engine/Plugins/Experimental/EaseCurveTool/Source/EaseCurveTool/Private/EaseCurveToolExtender.cpp
--- a/Engine/Plugins/Experimental/EaseCurveTool/Source/EaseCurveTool/Private/EaseCurveToolExtender.cpp
+++ b/Engine/Plugins/Experimental/EaseCurveTool/Source/EaseCurveTool/Private/EaseCurveToolExtender.cpp
@@ -18,7 +18,6 @@
 #include "ToolMenuSection.h"
 #include "Widgets/Layout/SExpandableArea.h"
 #include "Widgets/SNullWidget.h"
-#include "Widgets/Text/STextBlock.h"
 
 #define LOCTEXT_NAMESPACE "EaseCurveToolExtender"
 
@@ -50,9 +49,9 @@ FEaseCurveToolExtender::~FEaseCurveToolExtender()
 		SequencerModule.UnregisterOnSequencerCreated(SequencerCreatedHandle);
 	}
 
-	for (const TPair<FName, FEaseCurveToolInstance>& ToolInstance : ToolInstances)
+	for (TPair<FName, FEaseCurveToolInstance>& ToolInstance : ToolInstances)
 	{
-		//RemoveSequencerToolBarExtension(ToolInstance.Value);
+		RemoveSidebarExtension(ToolInstance.Value);
 	}
 }
 
@@ -77,9 +76,7 @@ FEaseCurveToolExtender::FEaseCurveToolInstance* FEaseCurveToolExtender::FindOrAd
 			OnSequencerClosed(NewToolId, InSequencer);
 		});
 
-	InSequencer->RegisterDrawerSection(TEXT("SelectionDetails"), MakeShared<FEaseCurveToolSidebarSection>(InSequencer));
-
-	//AddSidebarExtension(NewToolInstance);
+	AddSidebarExtension(NewToolInstance);
 
 	return &NewToolInstance;
 }
@@ -115,8 +112,6 @@ void FEaseCurveToolExtender::OnSequencerClosed(const FName InToolId, const TShar
 		ToolInstances[ToolId].SidebarExtender.Reset();
 	}
 
-	InSequencer->UnregisterDrawerSection(TEXT("SelectionDetails"), FEaseCurveToolSidebarSection::UniqueId);
-
 	RemoveSidebarExtension(ToolInstances[ToolId]);
 
 	ToolInstances.Remove(ToolId);
@@ -168,7 +163,7 @@ TSharedPtr<FEaseCurveTool> FEaseCurveToolExtender::FindToolInstanceByCurveEditor
 	{
 		if (const TSharedPtr<ISequencer> Sequencer = ToolInstance.Value.WeakSequencer.Pin())
 		{
-			if (const TSharedPtr<FCurveEditor> CurveEditor = GetCurveEditorFromSequencer(Sequencer.ToSharedRef()))
+			if (const TSharedPtr<FCurveEditor> CurveEditor = FEaseCurveToolUtils::GetCurveEditorFromSequencer(Sequencer.ToSharedRef()))
 			{
 				if (CurveEditor == InCurveEditor)
 				{
@@ -184,19 +179,16 @@ void FEaseCurveToolExtender::AddSidebarExtension(FEaseCurveToolInstance& InToolI
 {
 	check(InToolInstance.Instance.IsValid());
 
-	UToolMenu* const SidebarMenu = UToolMenus::Get()->ExtendMenu(TEXT("Sequencer.Sidebar"));
-	SidebarMenu->AddDynamicSection(NAME_None, FNewToolMenuDelegate::CreateStatic(&FEaseCurveToolExtender::CreateSidebarWidget, InToolInstance.ToolId));
-
-	/*InToolInstance.SidebarExtender = MakeShared<FExtender>();
+	InToolInstance.SidebarExtender = MakeShared<FExtender>();
 
 	InToolInstance.SidebarExtender->AddMenuExtension(TEXT("KeyEdit")
-		, EExtensionHook::First
+		, EExtensionHook::After
 		, InToolInstance.Instance->GetCommandList()
-		, FMenuExtensionDelegate::CreateStatic(&FEaseCurveToolExtender::CreateSidebarWidget, InToolInstance.ToolId)
+		, FMenuExtensionDelegate::CreateStatic(&FEaseCurveToolExtender::AddSidebarWidget, InToolInstance.ToolId)
 	);
 
 	ISequencerModule& SequencerModule = FModuleManager::LoadModuleChecked<ISequencerModule>(TEXT("Sequencer"));
-	SequencerModule.GetSidebarExtensibilityManager()->AddExtender(InToolInstance.SidebarExtender);*/
+	SequencerModule.GetSidebarExtensibilityManager()->AddExtender(InToolInstance.SidebarExtender);
 }
 
 void FEaseCurveToolExtender::RemoveSidebarExtension(FEaseCurveToolInstance& InToolInstance)
@@ -208,66 +200,7 @@ void FEaseCurveToolExtender::RemoveSidebarExtension(FEaseCurveToolInstance& InTo
 	}
 }
 
-void FEaseCurveToolExtender::CreateSidebarWidget(UToolMenu* const InMenu, const FName InToolId)
-{
-	const TSharedPtr<FEaseCurveTool> ToolInstance = GetToolInstance(InToolId);
-	if (!ToolInstance.IsValid())
-	{
-		return;
-	}
-
-	const TSharedRef<SExpandableArea> ToolMenuWidget =
-		SNew(SExpandableArea)
-		.AreaTitle(LOCTEXT("EaseCurveToolTitle", "Ease Curve Tool"))
-		.InitiallyCollapsed(false)
-		.BorderImage(FAppStyle::Get().GetBrush("Brushes.Header"))
-		.BodyBorderImage(FAppStyle::Get().GetBrush("Brushes.Recessed"))
-		.HeaderPadding(FMargin(2.f, 5.f))
-		.BodyContent()
-		[
-			ToolInstance->GenerateWidget()
-		];
-
-	/*FMenuEntryStyleParams StyleParams;
-	StyleParams.bNoIndent = true;
-	StyleParams.HorizontalAlignment = HAlign_Fill;
-	StyleParams.VerticalAlignment = VAlign_Top;*/
-
-	FToolMenuEntry ToolMenuEntry = FToolMenuEntry::InitWidget(TEXT("EaseCurveTool"), ToolMenuWidget, FText::GetEmpty());
-	ToolMenuEntry.Visibility = TAttribute<EVisibility>::CreateLambda([InToolId]()
-		{
-			if (const TSharedPtr<FEaseCurveTool> ToolInstance = GetToolInstance(InToolId))
-			{
-				if (const UEaseCurveToolSettings* const ToolSettings = GetDefault<UEaseCurveToolSettings>())
-				{
-					return ToolSettings->ShouldShowInSidebar() && !ToolInstance->IsToolTabVisible()
-						? EVisibility::Visible : EVisibility::Collapsed;
-				}
-			}
-			return EVisibility::Collapsed;
-		});
-	/*MenuBuilder.AddWidget(ToolMenuWidget
-		, FText::GetEmpty()
-		, StyleParams
-		, true
-		, FText::GetEmpty()
-		, FSlateIcon()
-		, TAttribute<EVisibility>::CreateLambda([InToolId]()
-			{
-				if (const TSharedPtr<FEaseCurveTool> ToolInstance = GetToolInstance(InToolId))
-				{
-					if (const UEaseCurveToolSettings* const ToolSettings = GetDefault<UEaseCurveToolSettings>())
-					{
-						return ToolSettings->ShouldShowInSidebar() && !ToolInstance->IsToolTabVisible()
-							? EVisibility::Visible : EVisibility::Collapsed;
-					}
-				}
-				return EVisibility::Collapsed;
-			})
-	);*/
-}
-
-void FEaseCurveToolExtender::CreateSidebarWidget(FMenuBuilder& MenuBuilder, const FName InToolId)
+void FEaseCurveToolExtender::AddSidebarWidget(FMenuBuilder& MenuBuilder, const FName InToolId)
 {
 	const TSharedPtr<FEaseCurveTool> ToolInstance = GetToolInstance(InToolId);
 	if (!ToolInstance.IsValid())
@@ -282,6 +215,7 @@ void FEaseCurveToolExtender::CreateSidebarWidget(FMenuBuilder& MenuBuilder, cons
 		.BorderImage(FAppStyle::Get().GetBrush("Brushes.Header"))
 		.BodyBorderImage(FAppStyle::Get().GetBrush("Brushes.Recessed"))
 		.HeaderPadding(FMargin(2.f, 5.f))
+		.Visibility_Static(&FEaseCurveToolExtender::GetSidebarVisibility, InToolId)
 		.BodyContent()
 		[
 			ToolInstance->GenerateWidget()
@@ -298,21 +232,25 @@ void FEaseCurveToolExtender::CreateSidebarWidget(FMenuBuilder& MenuBuilder, cons
 		, /*bInSearchable=*/true
 		, FText::GetEmpty()
 		, FSlateIcon()
-		, TAttribute<EVisibility>::CreateLambda([InToolId]()
-			{
-				if (const TSharedPtr<FEaseCurveTool> ToolInstance = GetToolInstance(InToolId))
-				{
-					if (const UEaseCurveToolSettings* const ToolSettings = GetDefault<UEaseCurveToolSettings>())
-					{
-						return ToolSettings->ShouldShowInSidebar() && !ToolInstance->IsToolTabVisible()
-							? EVisibility::Visible : EVisibility::Collapsed;
-					}
-				}
-				return EVisibility::Collapsed;
-			})
+		//, TAttribute<EVisibility>::CreateStatic(&FEaseCurveToolExtender::GetSidebarVisibility, InToolId)
 	);
 }
 
+EVisibility FEaseCurveToolExtender::GetSidebarVisibility(const FName InToolId)
+{
+	if (const TSharedPtr<FEaseCurveTool> ToolInstance = GetToolInstance(InToolId))
+	{
+		if (const UEaseCurveToolSettings* const ToolSettings = GetDefault<UEaseCurveToolSettings>())
+		{
+			return (ToolSettings->ShouldShowInSidebar()
+					&& !ToolInstance->IsToolTabVisible()
+					&& ToolInstance->HasCachedKeysToEase())
+				? EVisibility::Visible : EVisibility::Collapsed;
+		}
+	}
+	return EVisibility::Collapsed;
+}
+
 TSharedRef<SWidget> FEaseCurveToolExtender::MakeQuickPresetMenu(const FName InToolId)
 {
 	UToolMenus* ToolMenus = UToolMenus::Get();
diff --git a/Engine/Plugins/Experimental/EaseCurveTool/Source/EaseCurveTool/Private/EaseCurveToolExtender.h b/Engine/Plugins/Experimental/EaseCurveTool/Source/EaseCurveTool/Private/EaseCurveToolExtender.h
--- a/Engine/Plugins/Experimental/EaseCurveTool/Source/EaseCurveTool/Private/EaseCurveToolExtender.h
+++ b/Engine/Plugins/Experimental/EaseCurveTool/Source/EaseCurveTool/Private/EaseCurveToolExtender.h
@@ -9,9 +9,11 @@
 class FCurveEditor;
 class FExtender;
 class FMenuBuilder;
+class FName;
 class ISequencer;
 class SWidget;
 class UToolMenu;
+struct EVisibility;
 
 #define UE_API EASECURVETOOL_API
 
@@ -74,8 +76,9 @@ protected:
 	static void AddSidebarExtension(FEaseCurveToolInstance& InToolInstance);
 	static void RemoveSidebarExtension(FEaseCurveToolInstance& InToolInstance);
 
-	static void CreateSidebarWidget(UToolMenu* const InMenu, const FName InToolId);
-	static void CreateSidebarWidget(FMenuBuilder& MenuBuilder, const FName InToolId);
+	static void AddSidebarWidget(FMenuBuilder& MenuBuilder, const FName InToolId);
+
+	static EVisibility GetSidebarVisibility(const FName InToolId);
 
 private:
 	FDelegateHandle SequencerCreatedHandle;
diff --git a/Engine/Plugins/Experimental/EaseCurveTool/Source/EaseCurveTool/Private/EaseCurveToolUtils.cpp b/Engine/Plugins/Experimental/EaseCurveTool/Source/EaseCurveTool/Private/EaseCurveToolUtils.cpp
--- a/Engine/Plugins/Experimental/EaseCurveTool/Source/EaseCurveTool/Private/EaseCurveToolUtils.cpp
+++ b/Engine/Plugins/Experimental/EaseCurveTool/Source/EaseCurveTool/Private/EaseCurveToolUtils.cpp
@@ -2,14 +2,16 @@
 
 #include "EaseCurveToolUtils.h"
 #include "CurveEditorSelection.h"
+#include "EaseCurveTool.h"
 #include "MVVM/CurveEditorExtension.h"
 #include "MVVM/ViewModelPtr.h"
+#include "MVVM/ViewModels/ChannelModel.h"
 #include "MVVM/ViewModels/SequencerEditorViewModel.h"
 
 namespace UE::EaseCurveTool
 {
 
-UEaseCurveLibrary* GetToolPresetLibrary(const TWeakPtr<FEaseCurveTool>& InWeakTool)
+UEaseCurveLibrary* FEaseCurveToolUtils::GetToolPresetLibrary(const TWeakPtr<FEaseCurveTool>& InWeakTool)
 {
 	if (const TSharedPtr<FEaseCurveTool> Tool = InWeakTool.Pin())
 	{
@@ -18,7 +20,7 @@ UEaseCurveLibrary* GetToolPresetLibrary(const TWeakPtr<FEaseCurveTool>& InWeakTo
 	return nullptr;
 }
 
-bool CompareCurveEditorSelections(const FCurveEditorSelection& InSelectionA, const FCurveEditorSelection& InSelectionB)
+bool FEaseCurveToolUtils::CompareCurveEditorSelections(const FCurveEditorSelection& InSelectionA, const FCurveEditorSelection& InSelectionB)
 {
 	const TMap<FCurveModelID, FKeyHandleSet> SelectionMapA = InSelectionA.GetAll();
 	const TMap<FCurveModelID, FKeyHandleSet> SelectionMapB = InSelectionB.GetAll();
@@ -48,7 +50,7 @@ bool CompareCurveEditorSelections(const FCurveEditorSelection& InSelectionA, con
 	return true;
 }
 
-TSharedPtr<FCurveEditor> GetCurveEditorFromSequencer(const TSharedRef<ISequencer>& InSequencer)
+TSharedPtr<FCurveEditor> FEaseCurveToolUtils::GetCurveEditorFromSequencer(const TSharedRef<ISequencer>& InSequencer)
 {
 	using namespace Sequencer;
 
@@ -66,4 +68,58 @@ TSharedPtr<FCurveEditor> GetCurveEditorFromSequencer(const TSharedRef<ISequencer
 	return nullptr;
 }
 
+bool FEaseCurveToolUtils::FindChannelHandleFromCurveModel(const FCurveModel& InCurveModel, FMovieSceneChannelHandle& OutChannelHandle)
+{
+	UMovieSceneSection* const Section = InCurveModel.GetOwningObjectOrOuter<UMovieSceneSection>();
+	if (!Section)
+	{
+		return false;
+	}
+
+	const void* RawCurve = InCurveModel.GetCurve();
+	if (!RawCurve)
+	{
+		return false;
+	}
+
+	FMovieSceneChannelProxy& ChannelProxy = Section->GetChannelProxy();
+	const TWeakPtr<FMovieSceneChannelProxy> WeakProxy = ChannelProxy.AsWeak();
+
+	for (const FMovieSceneChannelEntry& Entry : ChannelProxy.GetAllEntries())
+	{
+		const TArrayView<FMovieSceneChannel* const> Channels = Entry.GetChannels();
+		for (int32 Index = 0; Index < Channels.Num(); ++Index)
+		{
+			if (Channels[Index] == RawCurve)
+			{
+				OutChannelHandle = FMovieSceneChannelHandle(WeakProxy, Entry.GetChannelTypeName(), Index);
+				return true;
+			}
+		}
+	}
+
+	return false;
+}
+
+bool FEaseCurveToolUtils::HasWeightedBrokenTangents(const FRichCurveKey& InKey)
+{
+	return InKey.TangentMode == RCTM_Break
+		&& InKey.TangentWeightMode == RCTWM_WeightedBoth
+		&& InKey.InterpMode == RCIM_Cubic;
+}
+
+bool FEaseCurveToolUtils::HasWeightedBrokenTangents(const FMovieSceneDoubleValue& InValue)
+{
+	return InValue.TangentMode == RCTM_Break
+		&& InValue.Tangent.TangentWeightMode == RCTWM_WeightedBoth
+		&& InValue.InterpMode == RCIM_Cubic;
+}
+
+bool FEaseCurveToolUtils::HasWeightedBrokenTangents(const FMovieSceneFloatValue& InValue)
+{
+	return InValue.TangentMode == RCTM_Break
+		&& InValue.Tangent.TangentWeightMode == RCTWM_WeightedBoth
+		&& InValue.InterpMode == RCIM_Cubic;
+}
+
 } // namespace UE::EaseCurveTool
diff --git a/Engine/Plugins/Experimental/EaseCurveTool/Source/EaseCurveTool/Private/EaseCurveToolUtils.h b/Engine/Plugins/Experimental/EaseCurveTool/Source/EaseCurveTool/Private/EaseCurveToolUtils.h
--- a/Engine/Plugins/Experimental/EaseCurveTool/Source/EaseCurveTool/Private/EaseCurveToolUtils.h
+++ b/Engine/Plugins/Experimental/EaseCurveTool/Source/EaseCurveTool/Private/EaseCurveToolUtils.h
@@ -3,22 +3,34 @@
 #pragma once
 
 #include "CurveEditor.h"
-#include "EaseCurveTool.h"
 #include "ISequencer.h"
 
 class UEaseCurveLibrary;
 struct FCurveEditorSelection;
+struct FMovieSceneFloatValue;
+struct FMovieSceneDoubleValue;
+struct FRichCurveKey;
 
 namespace UE::EaseCurveTool
 {
 
 class FEaseCurveTool;
 
-UEaseCurveLibrary* GetToolPresetLibrary(const TWeakPtr<FEaseCurveTool>& InWeakTool);
+class FEaseCurveToolUtils
+{
+public:
+	static UEaseCurveLibrary* GetToolPresetLibrary(const TWeakPtr<FEaseCurveTool>& InWeakTool);
+
+	/** Returns True if the curve editor selections match */
+	static bool CompareCurveEditorSelections(const FCurveEditorSelection& InSelectionA, const FCurveEditorSelection& InSelectionB);
+
+	static TSharedPtr<FCurveEditor> GetCurveEditorFromSequencer(const TSharedRef<ISequencer>& InSequencer);
 
-/** Returns True if the curve editor selections match */
-bool CompareCurveEditorSelections(const FCurveEditorSelection& InSelectionA, const FCurveEditorSelection& InSelectionB);
+	static bool FindChannelHandleFromCurveModel(const FCurveModel& InCurveModel, FMovieSceneChannelHandle& OutChannelHandle);
 
-TSharedPtr<FCurveEditor> GetCurveEditorFromSequencer(const TSharedRef<ISequencer>& InSequencer);
+	static bool HasWeightedBrokenTangents(const FRichCurveKey& InKey);
+	static bool HasWeightedBrokenTangents(const FMovieSceneDoubleValue& InValue);
+	static bool HasWeightedBrokenTangents(const FMovieSceneFloatValue& InValue);
+};
 
 } // namespace UE::EaseCurveTool
diff --git a/Engine/Plugins/Experimental/EaseCurveTool/Source/EaseCurveTool/Private/Widgets/SEaseCurveAddNew.cpp b/Engine/Plugins/Experimental/EaseCurveTool/Source/EaseCurveTool/Private/Widgets/SEaseCurveAddNew.cpp
--- a/Engine/Plugins/Experimental/EaseCurveTool/Source/EaseCurveTool/Private/Widgets/SEaseCurveAddNew.cpp
+++ b/Engine/Plugins/Experimental/EaseCurveTool/Source/EaseCurveTool/Private/Widgets/SEaseCurveAddNew.cpp
@@ -81,7 +81,10 @@ bool SEaseCurveAddNew::IsButtonEnabled() const
 		const int32 SelectedKeyCount = KeySelection.GetTotalSelectedKeys();
 		const int32 KeyChannelCount = KeySelection.GetKeyChannelCount();
 		const bool bValidKeyCount = KeyChannelCount == 1 && SelectedKeyCount > 0 && SelectedKeyCount < 3;
-		return bValidKeyCount && Tool->GetPresetLibrary() && Tool->AreAllEaseCurves();
+		return bValidKeyCount
+			&& Tool->GetPresetLibrary()
+			&& Tool->HasCachedKeysToEase()
+			&& Tool->GetSelectionError() == EEaseCurveToolError::None;
 	}
 	return false;
 }
@@ -122,7 +125,7 @@ void SEaseCurveAddNew::OnNewPresetTextCommitted(const FText& InNewText, ETextCom
 	{
 		if (const TSharedPtr<FEaseCurveTool> Tool = WeakTool.Pin())
 		{
-			if (UEaseCurveLibrary* const PresetLibrary = GetToolPresetLibrary(WeakTool))
+			if (UEaseCurveLibrary* const PresetLibrary = FEaseCurveToolUtils::GetToolPresetLibrary(WeakTool))
 			{
 				if (const TSharedPtr<ISequencer> Sequencer = Tool->GetSequencer())
 				{
diff --git a/Engine/Plugins/Experimental/EaseCurveTool/Source/EaseCurveTool/Private/Widgets/SEaseCurveEditor.cpp b/Engine/Plugins/Experimental/EaseCurveTool/Source/EaseCurveTool/Private/Widgets/SEaseCurveEditor.cpp
--- a/Engine/Plugins/Experimental/EaseCurveTool/Source/EaseCurveTool/Private/Widgets/SEaseCurveEditor.cpp
+++ b/Engine/Plugins/Experimental/EaseCurveTool/Source/EaseCurveTool/Private/Widgets/SEaseCurveEditor.cpp
@@ -87,8 +87,9 @@ void SEaseCurveEditor::Construct(const FArguments& InArgs, const TSharedRef<FEas
 	CurveThickness = InArgs._CurveThickness;
 	CurveColor = InArgs._CurveColor;
 	Operation = InArgs._Operation;
-	ShowEqualValueKeyError = InArgs._ShowEqualValueKeyError;
-	IsEaseCurveSelection = InArgs._IsEaseCurveSelection;
+
+	bCanEditCurve = InArgs._CanEditCurve;
+	ErrorMessage = InArgs._ErrorMessage;
 
 	StartText = InArgs._StartText;
 	StartTooltipText = InArgs._StartTooltipText;
@@ -201,11 +202,7 @@ int32 SEaseCurveEditor::OnPaint(const FPaintArgs& InArgs
 			, InText, FontInfo, DrawEffects, FStyleColors::Foreground.GetSpecifiedColor());
 	};
 
-	if (ShowEqualValueKeyError.Get(false))
-	{
-		PaintCenteredText(LOCTEXT("EqualValueKeys", "No different key values!"));
-	}
-	else
+	if (CanEditCurve())
 	{
 		InLayerId = PaintGrid(ScaleInfo, InGeometry, OutDrawElements, ++InLayerId, InCullingRect, DrawEffects);
 
@@ -243,26 +240,26 @@ int32 SEaseCurveEditor::OnPaint(const FPaintArgs& InArgs
 		PaintCurve(ScaleInfo, InGeometry, OutDrawElements, ++InLayerId, InCullingRect, DrawEffects, InWidgetStyle);
 
 		InLayerId = PaintKeys(ScaleInfo, InGeometry, OutDrawElements, ++InLayerId, DrawEffects, InWidgetStyle);
+	}
 
-		if (!IsEaseCurveSelection.Get(false))
-		{
-			PaintCenteredText(LOCTEXT("NoEaseCurve", "No weighted, broken, cubic tangents!"));
-		}
+	if (HasDrawErrorMessage())
+	{
+		PaintCenteredText(ErrorMessage.Get().GetValue());
+	}
 
-		if (DragState == EDragState::MarqueeSelect)
-		{
-			const FVector2D MarqueTopLeft(FMath::Min(MouseDownLocation.X, MouseMoveLocation.X), FMath::Min(MouseDownLocation.Y, MouseMoveLocation.Y));
-			const FVector2D MarqueBottomRight(FMath::Max(MouseDownLocation.X, MouseMoveLocation.X), FMath::Max(MouseDownLocation.Y, MouseMoveLocation.Y));
+	if (Tool->HasCachedKeysToEase())
+	{
+		PaintSelectionText(ScaleInfo, InGeometry, OutDrawElements, ++InLayerId, DrawEffects);
+	}
 
-			FSlateDrawElement::MakeBox(OutDrawElements, ++InLayerId
-				, InGeometry.ToPaintGeometry(MarqueBottomRight - MarqueTopLeft, FSlateLayoutTransform(MarqueTopLeft))
-				, FAppStyle::GetBrush(TEXT("MarqueeSelection")));
-		}
+	if (DragState == EDragState::MarqueeSelect)
+	{
+		const FVector2D MarqueTopLeft(FMath::Min(MouseDownLocation.X, MouseMoveLocation.X), FMath::Min(MouseDownLocation.Y, MouseMoveLocation.Y));
+		const FVector2D MarqueBottomRight(FMath::Max(MouseDownLocation.X, MouseMoveLocation.X), FMath::Max(MouseDownLocation.Y, MouseMoveLocation.Y));
 
-		if (Tool->HasCachedKeysToEase())
-		{
-			PaintSelectionText(ScaleInfo, InGeometry, OutDrawElements, ++InLayerId, DrawEffects);
-		}
+		FSlateDrawElement::MakeBox(OutDrawElements, ++InLayerId
+			, InGeometry.ToPaintGeometry(MarqueBottomRight - MarqueTopLeft, FSlateLayoutTransform(MarqueTopLeft))
+			, FAppStyle::GetBrush(TEXT("MarqueeSelection")));
 	}
 
 	return InLayerId;
@@ -455,7 +452,7 @@ void SEaseCurveEditor::CreateLinesForSegment(const ERichCurveInterpMode InInterp
 {
 	static constexpr FLinearColor DisabledFadeFactor(0.03f, 0.03f, 0.03f, 1.f);
 
-	const bool bIsEaseCurve = IsEaseCurveSelection.Get(false);
+	const bool bIsEaseCurve = CanEditCurve();
 
 	FLinearColor FadedCurveColor = CurveColor;
 	FadedCurveColor.A = 0.05f;
@@ -868,8 +865,8 @@ void SEaseCurveEditor::UpdateCurveToolTip(const FGeometry& InMyGeometry, const F
 	CurveToolTipInputText = FText::Format(LOCTEXT("CurveToolTipTimeFormat", "Time:\t{0}"), FText::AsNumber(TimeValueRange.X, &FormattingOptions));
 
 	if (DragState != EDragState::None
-		|| ShowEqualValueKeyError.Get(false)
-		|| !IsEaseCurveSelection.Get(false))
+		|| HasDrawErrorMessage()
+		|| !CanEditCurve())
 	{
 		ToolTipIndex = INDEX_NONE;
 		CurveToolTip.Reset();
@@ -973,6 +970,11 @@ FReply SEaseCurveEditor::OnKeyDown(const FGeometry& MyGeometry, const FKeyEvent&
 
 void SEaseCurveEditor::TryStartDrag(const FGeometry& InMyGeometry, const FPointerEvent& InMouseEvent)
 {
+	if (!CanEditCurve())
+	{
+		return;
+	}
+
 	const bool bLeftMouseButton = InMouseEvent.IsMouseButtonDown(EKeys::LeftMouseButton);
 	const bool bMiddleMouseButton = InMouseEvent.IsMouseButtonDown(EKeys::MiddleMouseButton);
 	const bool bRightMouseButton = InMouseEvent.IsMouseButtonDown(EKeys::RightMouseButton);
@@ -1510,6 +1512,16 @@ FVector2D SEaseCurveEditor::SnapLocation(const FVector2D& InLocation) const
 	return OutLocation;
 }
 
+bool SEaseCurveEditor::HasDrawErrorMessage() const
+{
+	return ErrorMessage.IsSet() && ErrorMessage.Get().IsSet();
+}
+
+bool SEaseCurveEditor::CanEditCurve() const
+{
+	return bCanEditCurve.Get(false);
+}
+
 } // namespace UE::EaseCurveTool
 
 #undef LOCTEXT_NAMESPACE
diff --git a/Engine/Plugins/Experimental/EaseCurveTool/Source/EaseCurveTool/Private/Widgets/SEaseCurveEditor.h b/Engine/Plugins/Experimental/EaseCurveTool/Source/EaseCurveTool/Private/Widgets/SEaseCurveEditor.h
--- a/Engine/Plugins/Experimental/EaseCurveTool/Source/EaseCurveTool/Private/Widgets/SEaseCurveEditor.h
+++ b/Engine/Plugins/Experimental/EaseCurveTool/Source/EaseCurveTool/Private/Widgets/SEaseCurveEditor.h
@@ -3,7 +3,6 @@
 #pragma once
 
 #include "Containers/UnrealString.h"
-#include "Curves/CurveBase.h"
 #include "Curves/KeyHandle.h"
 #include "EaseCurveTangents.h"
 #include "EaseCurveTool.h"
@@ -11,14 +10,12 @@
 #include "Framework/SlateDelegates.h"
 #include "Input/Reply.h"
 #include "Layout/Geometry.h"
-#include "Layout/Visibility.h"
 #include "Misc/Attribute.h"
 #include "SCurveEditor.h"
 #include "Templates/SharedPointer.h"
 #include "UObject/GCObject.h"
 #include "Widgets/DeclarativeSyntaxSupport.h"
 #include "Widgets/SCompoundWidget.h"
-#include "Widgets/SWidget.h"
 
 class IToolTip;
 
@@ -45,8 +42,6 @@ public:
 		, _ExtendedGridColor(FLinearColor(0.02f, 0.02f, 0.02f, 0.5f))
 		, _CurveThickness(2.5f)
 		, _CurveColor(FLinearColor::White)
-		, _ShowEqualValueKeyError(false)
-		, _IsEaseCurveSelection(true)
 	{
 		_Clipping = EWidgetClipping::ClipToBounds;
 	}
@@ -68,8 +63,9 @@ public:
 		SLATE_ARGUMENT(float, CurveThickness)
 		SLATE_ARGUMENT(FLinearColor, CurveColor)
 		SLATE_ATTRIBUTE(EEaseCurveToolOperation, Operation)
-		SLATE_ATTRIBUTE(bool, ShowEqualValueKeyError)
-		SLATE_ATTRIBUTE(bool, IsEaseCurveSelection)
+
+		SLATE_ATTRIBUTE(bool, CanEditCurve)
+		SLATE_ATTRIBUTE(TOptional<FText>, ErrorMessage)
 
 		SLATE_ATTRIBUTE(FText, StartText)
 		SLATE_ATTRIBUTE(FText, StartTooltipText)
@@ -97,6 +93,9 @@ public:
 
 	FKeyHandle GetSelectedKeyHandle() const { return SelectedTangent.KeyHandle; }
 
+	bool HasDrawErrorMessage() const;
+	bool CanEditCurve() const;
+
 protected:
 	/** Represents the different states of a drag operation */
 	enum class EDragState
@@ -286,8 +285,9 @@ protected:
 	float CurveThickness = 1.f;
 	FLinearColor CurveColor = FLinearColor::White;
 	TAttribute<EEaseCurveToolOperation> Operation;
-	TAttribute<bool> ShowEqualValueKeyError;
-	TAttribute<bool> IsEaseCurveSelection;
+
+	TAttribute<bool> bCanEditCurve;
+	TAttribute<TOptional<FText>> ErrorMessage;
 
 	TAttribute<FText> StartText;
 	TAttribute<FText> StartTooltipText;
diff --git a/Engine/Plugins/Experimental/EaseCurveTool/Source/EaseCurveTool/Private/Widgets/SEaseCurvePreset.cpp b/Engine/Plugins/Experimental/EaseCurveTool/Source/EaseCurveTool/Private/Widgets/SEaseCurvePreset.cpp
--- a/Engine/Plugins/Experimental/EaseCurveTool/Source/EaseCurveTool/Private/Widgets/SEaseCurvePreset.cpp
+++ b/Engine/Plugins/Experimental/EaseCurveTool/Source/EaseCurveTool/Private/Widgets/SEaseCurvePreset.cpp
@@ -184,7 +184,7 @@ void SEaseCurvePreset::OnNewPresetTextCommitted(const FText& InNewText, ETextCom
 		FEaseCurveTangents NewTangents;
 		if (OnGetNewPresetTangents.IsBound() && OnGetNewPresetTangents.Execute(NewTangents))
 		{
-			if (UEaseCurveLibrary* const PresetLibrary = GetToolPresetLibrary(WeakTool))
+			if (UEaseCurveLibrary* const PresetLibrary = FEaseCurveToolUtils::GetToolPresetLibrary(WeakTool))
 			{
 				FEaseCurvePreset NewPreset;
 				if (PresetLibrary->AddPresetToNewCategory(InNewText, NewTangents, NewPreset))
diff --git a/Engine/Plugins/Experimental/EaseCurveTool/Source/EaseCurveTool/Private/Widgets/SEaseCurvePresetList.cpp b/Engine/Plugins/Experimental/EaseCurveTool/Source/EaseCurveTool/Private/Widgets/SEaseCurvePresetList.cpp
--- a/Engine/Plugins/Experimental/EaseCurveTool/Source/EaseCurveTool/Private/Widgets/SEaseCurvePresetList.cpp
+++ b/Engine/Plugins/Experimental/EaseCurveTool/Source/EaseCurveTool/Private/Widgets/SEaseCurvePresetList.cpp
@@ -52,7 +52,8 @@ void SEaseCurvePresetList::Construct(const FArguments& InArgs, const TSharedRef<
 			SNew(SBox)
 			.Visibility_Lambda([this]()
 				{
-					return (GetToolPresetLibrary(WeakTool) != nullptr) ? EVisibility::Visible : EVisibility::Collapsed;
+					return (FEaseCurveToolUtils::GetToolPresetLibrary(WeakTool) != nullptr)
+						? EVisibility::Visible : EVisibility::Collapsed;
 				})
 			[
 				GenerateSearchRowWidget()
@@ -148,7 +149,7 @@ void SEaseCurvePresetList::Reload()
 
 void SEaseCurvePresetList::ReloadPresetItems()
 {
-	UEaseCurveLibrary* const PresetLibrary = GetToolPresetLibrary(WeakTool);
+	UEaseCurveLibrary* const PresetLibrary = FEaseCurveToolUtils::GetToolPresetLibrary(WeakTool);
 	if (!PresetLibrary)
 	{
 		return;
@@ -174,7 +175,7 @@ void SEaseCurvePresetList::RegenerateGroupWrapBox()
 		.EvenRowDistribution(true)
 		.NumColumnsOverride_Lambda([this]()
 			{
-				if (UEaseCurveLibrary* const PresetLibrary = GetToolPresetLibrary(WeakTool))
+				if (UEaseCurveLibrary* const PresetLibrary = FEaseCurveToolUtils::GetToolPresetLibrary(WeakTool))
 				{
 					const int32 EaseCurveCategoryCount = PresetLibrary->GetCategories().Num();
 					return FMath::Min(5, EaseCurveCategoryCount);
@@ -182,7 +183,7 @@ void SEaseCurvePresetList::RegenerateGroupWrapBox()
 				return 0;
 			});
 
-	if (UEaseCurveLibrary* const PresetLibrary = GetToolPresetLibrary(WeakTool))
+	if (UEaseCurveLibrary* const PresetLibrary = FEaseCurveToolUtils::GetToolPresetLibrary(WeakTool))
 	{
 		const TArray<FText>& PresetCategories = PresetLibrary->GetCategories();
 		const int32 CurvePresetCount = PresetCategories.Num();
@@ -251,7 +252,7 @@ void SEaseCurvePresetList::UpdateGroupsContent()
 			];
 	};
 
-	UEaseCurveLibrary* const PresetLibrary = GetToolPresetLibrary(WeakTool);
+	UEaseCurveLibrary* const PresetLibrary = FEaseCurveToolUtils::GetToolPresetLibrary(WeakTool);
 	if (!PresetLibrary)
 	{
 		GroupWidgetsParent->SetContent(GenerateNoPresetsWidget(LOCTEXT("NoLibrarySelected", "No library selected")));
@@ -407,7 +408,7 @@ void SEaseCurvePresetList::HandleSearchTextChanged(const FText& InSearchText)
 
 bool SEaseCurvePresetList::HandleCategoryDelete(const FText& InCategory)
 {
-	UEaseCurveLibrary* const PresetLibrary = GetToolPresetLibrary(WeakTool);
+	UEaseCurveLibrary* const PresetLibrary = FEaseCurveToolUtils::GetToolPresetLibrary(WeakTool);
 	if (!PresetLibrary)
 	{
 		return false;
@@ -425,7 +426,7 @@ bool SEaseCurvePresetList::HandleCategoryDelete(const FText& InCategory)
 
 bool SEaseCurvePresetList::HandleCategoryRename(const FText& InCategory, const FText& InNewName)
 {
-	UEaseCurveLibrary* const PresetLibrary = GetToolPresetLibrary(WeakTool);
+	UEaseCurveLibrary* const PresetLibrary = FEaseCurveToolUtils::GetToolPresetLibrary(WeakTool);
 	if (!PresetLibrary)
 	{
 		return false;
@@ -448,7 +449,7 @@ bool SEaseCurvePresetList::HandlePresetDelete(const TSharedPtr<FEaseCurvePreset>
 		return false;
 	}
 
-	UEaseCurveLibrary* const PresetLibrary = GetToolPresetLibrary(WeakTool);
+	UEaseCurveLibrary* const PresetLibrary = FEaseCurveToolUtils::GetToolPresetLibrary(WeakTool);
 	if (!PresetLibrary)
 	{
 		return false;
@@ -471,7 +472,7 @@ bool SEaseCurvePresetList::HandlePresetRename(const TSharedPtr<FEaseCurvePreset>
 		return false;
 	}
 
-	UEaseCurveLibrary* const PresetLibrary = GetToolPresetLibrary(WeakTool);
+	UEaseCurveLibrary* const PresetLibrary = FEaseCurveToolUtils::GetToolPresetLibrary(WeakTool);
 	if (!PresetLibrary)
 	{
 		return false;
@@ -507,7 +508,7 @@ bool SEaseCurvePresetList::HandleEndPresetMove(const TSharedPtr<FEaseCurvePreset
 		return false;
 	}
 
-	UEaseCurveLibrary* const PresetLibrary = GetToolPresetLibrary(WeakTool);
+	UEaseCurveLibrary* const PresetLibrary = FEaseCurveToolUtils::GetToolPresetLibrary(WeakTool);
 	if (!PresetLibrary)
 	{
 		return false;
@@ -561,7 +562,7 @@ FReply SEaseCurvePresetList::OnDeletePresetClick()
 		return FReply::Handled();
 	}
 
-	UEaseCurveLibrary* const PresetLibrary = GetToolPresetLibrary(WeakTool);
+	UEaseCurveLibrary* const PresetLibrary = FEaseCurveToolUtils::GetToolPresetLibrary(WeakTool);
 	if (!PresetLibrary)
 	{
 		return FReply::Handled();
@@ -580,7 +581,7 @@ FReply SEaseCurvePresetList::OnDeletePresetClick()
 
 FReply SEaseCurvePresetList::CreateNewCategory()
 {
-	UEaseCurveLibrary* const PresetLibrary = GetToolPresetLibrary(WeakTool);
+	UEaseCurveLibrary* const PresetLibrary = FEaseCurveToolUtils::GetToolPresetLibrary(WeakTool);
 	if (!PresetLibrary)
 	{
 		return FReply::Handled();
diff --git a/Engine/Plugins/Experimental/EaseCurveTool/Source/EaseCurveTool/Private/Widgets/SEaseCurveTool.cpp b/Engine/Plugins/Experimental/EaseCurveTool/Source/EaseCurveTool/Private/Widgets/SEaseCurveTool.cpp
--- a/Engine/Plugins/Experimental/EaseCurveTool/Source/EaseCurveTool/Private/Widgets/SEaseCurveTool.cpp
+++ b/Engine/Plugins/Experimental/EaseCurveTool/Source/EaseCurveTool/Private/Widgets/SEaseCurveTool.cpp
@@ -44,48 +44,45 @@ void SEaseCurveTool::Construct(const FArguments& InArgs, const TSharedRef<FEaseC
 			+ SVerticalBox::Slot()
 			.AutoHeight()
 			[
-				SNew(SBox)
+				SNew(SVerticalBox)
 				.Visibility_Lambda([this]()
-				{
-					if (const auto Tool = WeakTool.Pin())
 					{
-						return Tool->GetPresetLibrary() ? EVisibility::Visible : EVisibility::Collapsed;
-					}
-					return EVisibility::Collapsed;
-				})
+						if (const TSharedPtr<FEaseCurveTool> Tool = WeakTool.Pin())
+						{
+							return Tool->GetPresetLibrary() ? EVisibility::Visible : EVisibility::Collapsed;
+						}
+						return EVisibility::Collapsed;
+					})
+				+ SVerticalBox::Slot()
+				.AutoHeight()
+				.HAlign(HAlign_Fill)
+				.Padding(0.f, 1.f, 0.f, 0.f)
 				[
-					SNew(SVerticalBox)
-					+ SVerticalBox::Slot()
-					.AutoHeight()
-					.HAlign(HAlign_Fill)
-					.Padding(0.f, 1.f, 0.f, 0.f)
-					[
-						SAssignNew(CurvePresetWidget, SEaseCurvePreset, InTool)
-						.OnPresetChanged(this, &SEaseCurveTool::OnPresetChanged)
-						.OnQuickPresetChanged(this, &SEaseCurveTool::OnQuickPresetChanged)
-						.OnGetNewPresetTangents(this, &SEaseCurveTool::OnGetNewPresetTangents)			
-					]
-					+ SVerticalBox::Slot()
-					.AutoHeight()
-					.HAlign(HAlign_Fill)
-					.Padding(0.f, 4.f, 0.f, 0.f)
-					[
-						ConstructCurveEditorPanel()
-					]
-					+ SVerticalBox::Slot()
-					.AutoHeight()
-					.HAlign(HAlign_Fill)
-					.Padding(0.f, 3.f, 0.f, 0.f)
-					[
-						SNew(SEaseCurveTangents)
-						.InitialTangents(GetTangents())
-						.OnStartTangentChanged(this, &SEaseCurveTool::OnStartTangentSpinBoxChanged)
-						.OnStartWeightChanged(this, &SEaseCurveTool::OnStartTangentWeightSpinBoxChanged)
-						.OnEndTangentChanged(this, &SEaseCurveTool::OnEndTangentSpinBoxChanged)
-						.OnEndWeightChanged(this, &SEaseCurveTool::OnEndTangentWeightSpinBoxChanged)
-						.OnBeginSliderMovement(this, &SEaseCurveTool::OnBeginSliderMovement)
-						.OnEndSliderMovement(this, &SEaseCurveTool::OnEndSliderMovement)
-					]
+					SAssignNew(CurvePresetWidget, SEaseCurvePreset, InTool)
+					.OnPresetChanged(this, &SEaseCurveTool::OnPresetChanged)
+					.OnQuickPresetChanged(this, &SEaseCurveTool::OnQuickPresetChanged)
+					.OnGetNewPresetTangents(this, &SEaseCurveTool::OnGetNewPresetTangents)			
+				]
+				+ SVerticalBox::Slot()
+				.AutoHeight()
+				.HAlign(HAlign_Fill)
+				.Padding(0.f, 4.f, 0.f, 0.f)
+				[
+					ConstructCurveEditorPanel()
+				]
+				+ SVerticalBox::Slot()
+				.AutoHeight()
+				.HAlign(HAlign_Fill)
+				.Padding(0.f, 3.f, 0.f, 0.f)
+				[
+					SNew(SEaseCurveTangents)
+					.InitialTangents(GetTangents())
+					.OnStartTangentChanged(this, &SEaseCurveTool::OnStartTangentSpinBoxChanged)
+					.OnStartWeightChanged(this, &SEaseCurveTool::OnStartTangentWeightSpinBoxChanged)
+					.OnEndTangentChanged(this, &SEaseCurveTool::OnEndTangentSpinBoxChanged)
+					.OnEndWeightChanged(this, &SEaseCurveTool::OnEndTangentWeightSpinBoxChanged)
+					.OnBeginSliderMovement(this, &SEaseCurveTool::OnBeginSliderMovement)
+					.OnEndSliderMovement(this, &SEaseCurveTool::OnEndSliderMovement)
 				]
 			]
 		];
@@ -120,8 +117,8 @@ TSharedRef<SWidget> SEaseCurveTool::ConstructCurveEditorPanel()
 				.DisplayRate(this, &SEaseCurveTool::GetDisplayRate)
 				.Operation(this, &SEaseCurveTool::GetToolOperation)
 				.DesiredSize(this, &SEaseCurveTool::GetEditorSize)
-				.ShowEqualValueKeyError(this, &SEaseCurveTool::ShowEqualValueKeyError)
-				.IsEaseCurveSelection(this, &SEaseCurveTool::IsEaseCurveSelection)
+				.CanEditCurve(this, &SEaseCurveTool::CanEditCurve)
+				.ErrorMessage(this, &SEaseCurveTool::GetErrorMessage)
 				.OnTangentsChanged(this, &SEaseCurveTool::HandleEditorTangentsChanged)
 				.GridSnap_UObject(GetDefault<UEaseCurveToolSettings>(), &UEaseCurveToolSettings::GetGridSnap)
 				.GridSize_UObject(GetDefault<UEaseCurveToolSettings>(), &UEaseCurveToolSettings::GetGridSize)
@@ -188,14 +185,24 @@ FVector2D SEaseCurveTool::GetEditorSize() const
 	return FVector2D(CurrentGraphSize);
 }
 
-bool SEaseCurveTool::ShowEqualValueKeyError() const
+bool SEaseCurveTool::CanEditCurve() const
 {
-	return !WeakTool.Pin()->HasCachedKeysToEase();
+	if (const TSharedPtr<FEaseCurveTool> Tool = WeakTool.Pin())
+	{
+		const EEaseCurveToolError SelectionError = Tool->GetSelectionError();
+		return SelectionError == EEaseCurveToolError::None
+			|| SelectionError == EEaseCurveToolError::NoWeightedBrokenCubicTangents;
+	}
+	return false;
 }
 
-bool SEaseCurveTool::IsEaseCurveSelection() const
+TOptional<FText> SEaseCurveTool::GetErrorMessage() const
 {
-	return WeakTool.Pin()->AreAllEaseCurves();
+	if (const TSharedPtr<FEaseCurveTool> Tool = WeakTool.Pin())
+	{
+		return Tool->GetSelectionErrorText();
+	}
+	return TOptional<FText>();
 }
 
 void SEaseCurveTool::HandleLibraryChanged(const TWeakObjectPtr<UEaseCurveLibrary> InWeakLibrary)
diff --git a/Engine/Plugins/Experimental/EaseCurveTool/Source/EaseCurveTool/Private/Widgets/SEaseCurveTool.h b/Engine/Plugins/Experimental/EaseCurveTool/Source/EaseCurveTool/Private/Widgets/SEaseCurveTool.h
--- a/Engine/Plugins/Experimental/EaseCurveTool/Source/EaseCurveTool/Private/Widgets/SEaseCurveTool.h
+++ b/Engine/Plugins/Experimental/EaseCurveTool/Source/EaseCurveTool/Private/Widgets/SEaseCurveTool.h
@@ -53,8 +53,9 @@ protected:
 	FFrameRate GetDisplayRate() const;
 	EEaseCurveToolOperation GetToolOperation() const;
 	FVector2D GetEditorSize() const;
-	bool ShowEqualValueKeyError() const;
-	bool IsEaseCurveSelection() const;
+
+	bool CanEditCurve() const;
+	TOptional<FText> GetErrorMessage() const;
 
 	void HandleLibraryChanged(const TWeakObjectPtr<UEaseCurveLibrary> InWeakLibrary);
 
