[EaseCurveTool] - Change plugin from experimental to beta
- Allow null to clear the active preset library

diff --git a/Engine/Plugins/Experimental/EaseCurveTool/EaseCurveTool.uplugin b/Engine/Plugins/Experimental/EaseCurveTool/EaseCurveTool.uplugin
--- a/Engine/Plugins/Experimental/EaseCurveTool/EaseCurveTool.uplugin
+++ b/Engine/Plugins/Experimental/EaseCurveTool/EaseCurveTool.uplugin
@@ -11,8 +11,8 @@
 	"MarketplaceURL": "",
 	"EnabledByDefault": true,
 	"CanContainContent": true,
-	"IsBetaVersion": false,
-	"IsExperimentalVersion": true,
+	"IsBetaVersion": true,
+	"IsExperimentalVersion": false,
 	"Installed": false,
 	"Modules": [
 		{
diff --git a/Engine/Plugins/Experimental/EaseCurveTool/Source/EaseCurveTool/Private/EaseCurveToolSettings.cpp b/Engine/Plugins/Experimental/EaseCurveTool/Source/EaseCurveTool/Private/EaseCurveToolSettings.cpp
--- a/Engine/Plugins/Experimental/EaseCurveTool/Source/EaseCurveTool/Private/EaseCurveToolSettings.cpp
+++ b/Engine/Plugins/Experimental/EaseCurveTool/Source/EaseCurveTool/Private/EaseCurveToolSettings.cpp
@@ -9,7 +9,6 @@ UEaseCurveToolSettings::UEaseCurveToolSettings()
 	CategoryName = TEXT("Sequencer");
 	SectionName = TEXT("Ease Curve Tool");
 
-	//DefaultPresetLibrary = FSoftObjectPath(TEXT("/EaseCurveTool/Content/DefaultPresetLibrary.uasset"));
 	DefaultPresetLibrary = FSoftObjectPath(TEXT("/EaseCurveTool/DefaultPresetLibrary.DefaultPresetLibrary"));
 
 	NewPresetCategory = LOCTEXT("NewPresetCategory", "Custom");
diff --git a/Engine/Plugins/Experimental/EaseCurveTool/Source/EaseCurveTool/Private/Widgets/SEaseCurveLibraryComboBox.cpp b/Engine/Plugins/Experimental/EaseCurveTool/Source/EaseCurveTool/Private/Widgets/SEaseCurveLibraryComboBox.cpp
--- a/Engine/Plugins/Experimental/EaseCurveTool/Source/EaseCurveTool/Private/Widgets/SEaseCurveLibraryComboBox.cpp
+++ b/Engine/Plugins/Experimental/EaseCurveTool/Source/EaseCurveTool/Private/Widgets/SEaseCurveLibraryComboBox.cpp
@@ -77,10 +77,8 @@ void SEaseCurveLibraryComboBox::HandleLibrarySelected(const FAssetData& InAssetD
 	}
 
 	// Allow null to clear the active preset library
-	if (UEaseCurveLibrary* const PresetLibrary = Cast<UEaseCurveLibrary>(InAssetData.GetAsset()))
-	{
-		Tool->SetPresetLibrary(PresetLibrary);
-	}
+	UEaseCurveLibrary* const PresetLibrary = Cast<UEaseCurveLibrary>(InAssetData.GetAsset());
+	Tool->SetPresetLibrary(PresetLibrary);
 }
 
 EVisibility SEaseCurveLibraryComboBox::GetOpenTabButtonVisibility() const
