[EaseCurveTool] - Remove delegate in destructor instead of add

diff --git a/Engine/Plugins/Experimental/EaseCurveTool/Source/EaseCurveTool/Private/EaseCurveToolExtender.cpp b/Engine/Plugins/Experimental/EaseCurveTool/Source/EaseCurveTool/Private/EaseCurveToolExtender.cpp
--- a/Engine/Plugins/Experimental/EaseCurveTool/Source/EaseCurveTool/Private/EaseCurveToolExtender.cpp
+++ b/Engine/Plugins/Experimental/EaseCurveTool/Source/EaseCurveTool/Private/EaseCurveToolExtender.cpp
@@ -69,12 +69,7 @@ FEaseCurveToolExtender::FEaseCurveToolInstance* FEaseCurveToolExtender::FindOrAd
 	NewToolInstance.ToolId = NewToolId;
 	NewToolInstance.Instance = MakeShared<FEaseCurveTool>(InSequencer);
 	NewToolInstance.WeakSequencer = InSequencer;
-
-	NewToolInstance.SequencerClosedHandle = InSequencer->OnCloseEvent().AddLambda(
-		[this, NewToolId](const TSharedRef<ISequencer> InSequencer)
-		{
-			OnSequencerClosed(NewToolId, InSequencer);
-		});
+	NewToolInstance.SequencerClosedHandle = InSequencer->OnCloseEvent().AddRaw(this, &FEaseCurveToolExtender::OnSequencerClosed, NewToolId);
 
 	AddSidebarExtension(NewToolInstance);
 
@@ -90,7 +85,7 @@ void FEaseCurveToolExtender::OnSequencerCreated(const TSharedRef<ISequencer> InS
 	}
 }
 
-void FEaseCurveToolExtender::OnSequencerClosed(const FName InToolId, const TSharedRef<ISequencer> InSequencer)
+void FEaseCurveToolExtender::OnSequencerClosed(const TSharedRef<ISequencer> InSequencer, const FName InToolId)
 {
 	const FName ToolId = GetToolInstanceId(InSequencer);
 
diff --git a/Engine/Plugins/Experimental/EaseCurveTool/Source/EaseCurveTool/Private/EaseCurveToolExtender.h b/Engine/Plugins/Experimental/EaseCurveTool/Source/EaseCurveTool/Private/EaseCurveToolExtender.h
--- a/Engine/Plugins/Experimental/EaseCurveTool/Source/EaseCurveTool/Private/EaseCurveToolExtender.h
+++ b/Engine/Plugins/Experimental/EaseCurveTool/Source/EaseCurveTool/Private/EaseCurveToolExtender.h
@@ -12,7 +12,6 @@ class FMenuBuilder;
 class FName;
 class ISequencer;
 class SWidget;
-class UToolMenu;
 struct EVisibility;
 
 #define UE_API EASECURVETOOL_API
@@ -71,7 +70,7 @@ protected:
 	FEaseCurveToolInstance* FindOrAddToolInstance_Internal(const TSharedRef<ISequencer>& InSequencer);
 
 	void OnSequencerCreated(const TSharedRef<ISequencer> InSequencer);
-	void OnSequencerClosed(const FName InToolId, const TSharedRef<ISequencer> InSequencer);
+	void OnSequencerClosed(const TSharedRef<ISequencer> InSequencer, const FName InToolId);
 
 	static void AddSidebarExtension(FEaseCurveToolInstance& InToolInstance);
 	static void RemoveSidebarExtension(FEaseCurveToolInstance& InToolInstance);
diff --git a/Engine/Plugins/Experimental/EaseCurveTool/Source/EaseCurveTool/Private/Widgets/SEaseCurveAddNew.cpp b/Engine/Plugins/Experimental/EaseCurveTool/Source/EaseCurveTool/Private/Widgets/SEaseCurveAddNew.cpp
--- a/Engine/Plugins/Experimental/EaseCurveTool/Source/EaseCurveTool/Private/Widgets/SEaseCurveAddNew.cpp
+++ b/Engine/Plugins/Experimental/EaseCurveTool/Source/EaseCurveTool/Private/Widgets/SEaseCurveAddNew.cpp
@@ -121,7 +121,9 @@ FReply SEaseCurveAddNew::OnNewPresetKeyDownHandler(const FGeometry& InGeometry,
 
 void SEaseCurveAddNew::OnNewPresetTextCommitted(const FText& InNewText, ETextCommit::Type InCommitType)
 {
-	if (InCommitType == ETextCommit::OnEnter && !InNewText.IsEmpty())
+	if (InCommitType == ETextCommit::OnEnter
+		&& !InNewText.IsEmpty()
+		&& IsButtonEnabled())
 	{
 		if (const TSharedPtr<FEaseCurveTool> Tool = WeakTool.Pin())
 		{
diff --git a/Engine/Plugins/Experimental/EaseCurveTool/Source/EaseCurveTool/Private/Widgets/SEaseCurvePreset.cpp b/Engine/Plugins/Experimental/EaseCurveTool/Source/EaseCurveTool/Private/Widgets/SEaseCurvePreset.cpp
--- a/Engine/Plugins/Experimental/EaseCurveTool/Source/EaseCurveTool/Private/Widgets/SEaseCurvePreset.cpp
+++ b/Engine/Plugins/Experimental/EaseCurveTool/Source/EaseCurveTool/Private/Widgets/SEaseCurvePreset.cpp
@@ -22,7 +22,7 @@ SEaseCurvePreset::~SEaseCurvePreset()
 {
 	if (const TSharedPtr<FEaseCurveTool> Tool = WeakTool.Pin())
 	{
-		Tool->OnPresetLibraryChanged().AddSP(this, &SEaseCurvePreset::HandleLibraryChanged);
+		Tool->OnPresetLibraryChanged().RemoveAll(this);
 	}
 
 	if (UEaseCurveLibrary* const CurrentLibrary = WeakCurrentLibrary.Get())
