Ease Curve Tool: Check for loaded Sequencer module because it could be unloading

#rnx

diff --git a/Engine/Plugins/Experimental/EaseCurveTool/Source/EaseCurveTool/Private/EaseCurveToolExtender.cpp b/Engine/Plugins/Experimental/EaseCurveTool/Source/EaseCurveTool/Private/EaseCurveToolExtender.cpp
--- a/Engine/Plugins/Experimental/EaseCurveTool/Source/EaseCurveTool/Private/EaseCurveToolExtender.cpp
+++ b/Engine/Plugins/Experimental/EaseCurveTool/Source/EaseCurveTool/Private/EaseCurveToolExtender.cpp
@@ -1,4 +1,4 @@
-﻿// Copyright Epic Games, Inc. All Rights Reserved.
+// Copyright Epic Games, Inc. All Rights Reserved.
 
 #include "EaseCurveToolExtender.h"
 #include "EaseCurveTool.h"
@@ -102,8 +102,11 @@ void FEaseCurveToolExtender::OnSequencerClosed(const TSharedRef<ISequencer> InSe
 
 	if (ToolInstances[ToolId].SidebarExtender.IsValid())
 	{
-		ISequencerModule& SequencerModule = FModuleManager::LoadModuleChecked<ISequencerModule>(TEXT("Sequencer"));
-		SequencerModule.GetSidebarExtensibilityManager()->RemoveExtender(ToolInstances[ToolId].SidebarExtender);
+		if (FModuleManager::Get().IsModuleLoaded(TEXT("Sequencer")))
+		{
+			ISequencerModule& SequencerModule = FModuleManager::LoadModuleChecked<ISequencerModule>(TEXT("Sequencer"));
+			SequencerModule.GetSidebarExtensibilityManager()->RemoveExtender(ToolInstances[ToolId].SidebarExtender);
+		}
 		ToolInstances[ToolId].SidebarExtender.Reset();
 	}
 
@@ -190,8 +193,11 @@ void FEaseCurveToolExtender::RemoveSidebarExtension(FEaseCurveToolInstance& InTo
 {
 	if (InToolInstance.SidebarExtender.IsValid())
 	{
-		ISequencerModule& SequencerModule = FModuleManager::LoadModuleChecked<ISequencerModule>(TEXT("Sequencer"));
-		SequencerModule.GetSidebarExtensibilityManager()->RemoveExtender(InToolInstance.SidebarExtender);
+		if (FModuleManager::Get().IsModuleLoaded(TEXT("Sequencer")))
+		{
+			ISequencerModule& SequencerModule = FModuleManager::LoadModuleChecked<ISequencerModule>(TEXT("Sequencer"));
+			SequencerModule.GetSidebarExtensibilityManager()->RemoveExtender(InToolInstance.SidebarExtender);
+		}
 	}
 }
 
