Ease Curve Tool: Check for loaded Sequencer module because it could be unloading

#rnx

diff --git a/Engine/Plugins/Experimental/EaseCurveTool/Source/EaseCurveTool/Private/EaseCurveToolExtender.cpp b/Engine/Plugins/Experimental/EaseCurveTool/Source/EaseCurveTool/Private/EaseCurveToolExtender.cpp
--- a/Engine/Plugins/Experimental/EaseCurveTool/Source/EaseCurveTool/Private/EaseCurveToolExtender.cpp
+++ b/Engine/Plugins/Experimental/EaseCurveTool/Source/EaseCurveTool/Private/EaseCurveToolExtender.cpp
@@ -43,10 +43,9 @@ FEaseCurveToolExtender::~FEaseCurveToolExtender()
 {
 	IModularFeatures::Get().UnregisterModularFeature(GetModularFeatureName(), this);
 
-	if (FModuleManager::Get().IsModuleLoaded(TEXT("Sequencer")))
+	if (ISequencerModule* SequencerModule = FModuleManager::GetModulePtr<ISequencerModule>("Sequencer"))
 	{
-		ISequencerModule& SequencerModule = FModuleManager::LoadModuleChecked<ISequencerModule>(TEXT("Sequencer"));
-		SequencerModule.UnregisterOnSequencerCreated(SequencerCreatedHandle);
+		SequencerModule->UnregisterOnSequencerCreated(SequencerCreatedHandle);
 	}
 
 	for (TPair<FName, FEaseCurveToolInstance>& ToolInstance : ToolInstances)
@@ -102,10 +101,9 @@ void FEaseCurveToolExtender::OnSequencerClosed(const TSharedRef<ISequencer> InSe
 
 	if (ToolInstances[ToolId].SidebarExtender.IsValid())
 	{
-		if (FModuleManager::Get().IsModuleLoaded(TEXT("Sequencer")))
+		if (ISequencerModule* SequencerModule = FModuleManager::GetModulePtr<ISequencerModule>("Sequencer"))
 		{
-			ISequencerModule& SequencerModule = FModuleManager::LoadModuleChecked<ISequencerModule>(TEXT("Sequencer"));
-			SequencerModule.GetSidebarExtensibilityManager()->RemoveExtender(ToolInstances[ToolId].SidebarExtender);
+			SequencerModule->GetSidebarExtensibilityManager()->RemoveExtender(ToolInstances[ToolId].SidebarExtender);
 		}
 		ToolInstances[ToolId].SidebarExtender.Reset();
 	}
@@ -193,10 +191,9 @@ void FEaseCurveToolExtender::RemoveSidebarExtension(FEaseCurveToolInstance& InTo
 {
 	if (InToolInstance.SidebarExtender.IsValid())
 	{
-		if (FModuleManager::Get().IsModuleLoaded(TEXT("Sequencer")))
+		if (ISequencerModule* SequencerModule = FModuleManager::GetModulePtr<ISequencerModule>("Sequencer"))
 		{
-			ISequencerModule& SequencerModule = FModuleManager::LoadModuleChecked<ISequencerModule>(TEXT("Sequencer"));
-			SequencerModule.GetSidebarExtensibilityManager()->RemoveExtender(InToolInstance.SidebarExtender);
+			SequencerModule->GetSidebarExtensibilityManager()->RemoveExtender(InToolInstance.SidebarExtender);
 		}
 	}
 }
