Fixed asset loading on rightclick in content browser

diff --git a/Engine/Plugins/Experimental/EaseCurveTool/Source/EaseCurveTool/Private/Menus/EaseCurveLibraryMenu.cpp b/Engine/Plugins/Experimental/EaseCurveTool/Source/EaseCurveTool/Private/Menus/EaseCurveLibraryMenu.cpp
--- a/Engine/Plugins/Experimental/EaseCurveTool/Source/EaseCurveTool/Private/Menus/EaseCurveLibraryMenu.cpp
+++ b/Engine/Plugins/Experimental/EaseCurveTool/Source/EaseCurveTool/Private/Menus/EaseCurveLibraryMenu.cpp
@@ -23,9 +23,12 @@ TSharedRef<FExtender> FEaseCurveLibraryMenu::OnExtendContentBrowserAssetSelectio
 
 	for (const FAssetData& SelectedAsset : InSelectedAssets)
 	{
-		if (UEaseCurveLibrary* const LibraryAsset = Cast<UEaseCurveLibrary>(SelectedAsset.GetAsset()))
+		if(SelectedAsset.GetClass() == UEaseCurveLibrary::StaticClass())
 		{
-			WeakLibraryAssets.Add(LibraryAsset);
+			if (UEaseCurveLibrary* const LibraryAsset = Cast<UEaseCurveLibrary>(SelectedAsset.GetAsset()))
+			{
+				WeakLibraryAssets.Add(LibraryAsset);
+			}
 		}
 	}
 
