[EaseCurveTool] Set tool setting default "ShowInSidebar" to true

diff --git a/Engine/Plugins/Experimental/EaseCurveTool/Source/EaseCurveTool/Private/EaseCurveToolSettings.h b/Engine/Plugins/Experimental/EaseCurveTool/Source/EaseCurveTool/Private/EaseCurveToolSettings.h
--- a/Engine/Plugins/Experimental/EaseCurveTool/Source/EaseCurveTool/Private/EaseCurveToolSettings.h
+++ b/Engine/Plugins/Experimental/EaseCurveTool/Source/EaseCurveTool/Private/EaseCurveToolSettings.h
@@ -58,7 +58,7 @@ public:
 private:
 	/** If true, displays the tool in the Sequencer sidebar drawer */
 	UPROPERTY(Config, EditAnywhere, Category = "General")
-	bool bShowInSidebar = false;
+	bool bShowInSidebar = true;
 
 	/** If true, shows the ease curve combo button in the curve editor toolbar */
 	UPROPERTY(Config, EditAnywhere, Category = "General")
