[EaseCurveTool] Remove key selection cache error checking to enable forcibly applying ease curves in the curve editor

diff --git a/Engine/Plugins/Experimental/EaseCurveTool/Source/EaseCurveTool/Private/EaseCurveEditorExtension.cpp b/Engine/Plugins/Experimental/EaseCurveTool/Source/EaseCurveTool/Private/EaseCurveEditorExtension.cpp
--- a/Engine/Plugins/Experimental/EaseCurveTool/Source/EaseCurveTool/Private/EaseCurveEditorExtension.cpp
+++ b/Engine/Plugins/Experimental/EaseCurveTool/Source/EaseCurveTool/Private/EaseCurveEditorExtension.cpp
@@ -144,8 +144,9 @@ bool FEaseCurveEditorExtension::IsButtonEnabled() const
 {
 	if (const TSharedPtr<FEaseCurveTool> ToolInstance = GetToolInstance())
 	{
-		return ToolInstance->HasCachedKeysToEase()
-			&& ToolInstance->GetSelectionError() == EEaseCurveToolError::None;
+		// In order to apply operations on keys other than cubic weighted, we skip checking
+		// ToolInstance->GetSelectionError() == EEaseCurveToolError::None
+		return ToolInstance->HasCachedKeysToEase();
 	}
 	return false;
 }
