[EaseCurveTool] - Check external modules validity before using

diff --git a/Engine/Plugins/Experimental/EaseCurveTool/Source/EaseCurveTool/Private/EaseCurveTool.cpp b/Engine/Plugins/Experimental/EaseCurveTool/Source/EaseCurveTool/Private/EaseCurveTool.cpp
--- a/Engine/Plugins/Experimental/EaseCurveTool/Source/EaseCurveTool/Private/EaseCurveTool.cpp
+++ b/Engine/Plugins/Experimental/EaseCurveTool/Source/EaseCurveTool/Private/EaseCurveTool.cpp
@@ -398,11 +398,15 @@ void FEaseCurveTool::UpdateEaseCurveFromKeyCache()
 
 UCurveBase* FEaseCurveTool::CreateCurveAsset() const
 {
-	const FAssetToolsModule& AssetToolsModule = FModuleManager::LoadModuleChecked<FAssetToolsModule>(TEXT("AssetTools"));
+	FAssetToolsModule* const AssetToolsModule = FModuleManager::GetModulePtr<FAssetToolsModule>(TEXT("AssetTools"));
+	if (!AssetToolsModule)
+	{
+		return nullptr;
+	}
 
 	FString OutNewPackageName;
 	FString OutNewAssetName;
-	AssetToolsModule.Get().CreateUniqueAssetName(TEXT("/Game/NewCurve"), TEXT(""), OutNewPackageName, OutNewAssetName);
+	AssetToolsModule->Get().CreateUniqueAssetName(TEXT("/Game/NewCurve"), TEXT(""), OutNewPackageName, OutNewAssetName);
 
 	const TSharedRef<SDlgPickAssetPath> NewAssetDialog =
 		SNew(SDlgPickAssetPath)
diff --git a/Engine/Plugins/Experimental/EaseCurveTool/Source/EaseCurveTool/Private/EaseCurveToolExtender.cpp b/Engine/Plugins/Experimental/EaseCurveTool/Source/EaseCurveTool/Private/EaseCurveToolExtender.cpp
--- a/Engine/Plugins/Experimental/EaseCurveTool/Source/EaseCurveTool/Private/EaseCurveToolExtender.cpp
+++ b/Engine/Plugins/Experimental/EaseCurveTool/Source/EaseCurveTool/Private/EaseCurveToolExtender.cpp
@@ -34,16 +34,18 @@ FEaseCurveToolExtender::FEaseCurveToolExtender()
 {
 	IModularFeatures::Get().RegisterModularFeature(GetModularFeatureName(), this);
 
-	ISequencerModule& SequencerModule = FModuleManager::LoadModuleChecked<ISequencerModule>(TEXT("Sequencer"));
-	SequencerCreatedHandle = SequencerModule.RegisterOnSequencerCreated(
-		FOnSequencerCreated::FDelegate::CreateRaw(this, &FEaseCurveToolExtender::OnSequencerCreated));
+	if (ISequencerModule* const SequencerModule = FModuleManager::GetModulePtr<ISequencerModule>(TEXT("Sequencer")))
+	{
+		SequencerCreatedHandle = SequencerModule->RegisterOnSequencerCreated(
+			FOnSequencerCreated::FDelegate::CreateRaw(this, &FEaseCurveToolExtender::OnSequencerCreated));
+	}
 }
 
 FEaseCurveToolExtender::~FEaseCurveToolExtender()
 {
 	IModularFeatures::Get().UnregisterModularFeature(GetModularFeatureName(), this);
 
-	if (ISequencerModule* SequencerModule = FModuleManager::GetModulePtr<ISequencerModule>("Sequencer"))
+	if (ISequencerModule* const SequencerModule = FModuleManager::GetModulePtr<ISequencerModule>(TEXT("Sequencer")))
 	{
 		SequencerModule->UnregisterOnSequencerCreated(SequencerCreatedHandle);
 	}
@@ -99,9 +101,14 @@ void FEaseCurveToolExtender::OnSequencerClosed(const TSharedRef<ISequencer> InSe
 		return;
 	}
 
+	if (ToolInstances[ToolId].SequencerClosedHandle.IsValid())
+	{
+		InSequencer->OnCloseEvent().Remove(ToolInstances[ToolId].SequencerClosedHandle);
+	}
+
 	if (ToolInstances[ToolId].SidebarExtender.IsValid())
 	{
-		if (ISequencerModule* SequencerModule = FModuleManager::GetModulePtr<ISequencerModule>("Sequencer"))
+		if (ISequencerModule* const SequencerModule = FModuleManager::GetModulePtr<ISequencerModule>(TEXT("Sequencer")))
 		{
 			SequencerModule->GetSidebarExtensibilityManager()->RemoveExtender(ToolInstances[ToolId].SidebarExtender);
 		}
@@ -183,18 +190,22 @@ void FEaseCurveToolExtender::AddSidebarExtension(FEaseCurveToolInstance& InToolI
 		, FMenuExtensionDelegate::CreateStatic(&FEaseCurveToolExtender::AddSidebarWidget, InToolInstance.ToolId)
 	);
 
-	ISequencerModule& SequencerModule = FModuleManager::LoadModuleChecked<ISequencerModule>(TEXT("Sequencer"));
-	SequencerModule.GetSidebarExtensibilityManager()->AddExtender(InToolInstance.SidebarExtender);
+	if (ISequencerModule* const SequencerModule = FModuleManager::GetModulePtr<ISequencerModule>(TEXT("Sequencer")))
+	{
+		SequencerModule->GetSidebarExtensibilityManager()->AddExtender(InToolInstance.SidebarExtender);
+	}
 }
 
 void FEaseCurveToolExtender::RemoveSidebarExtension(FEaseCurveToolInstance& InToolInstance)
 {
 	if (InToolInstance.SidebarExtender.IsValid())
 	{
-		if (ISequencerModule* SequencerModule = FModuleManager::GetModulePtr<ISequencerModule>("Sequencer"))
+		if (ISequencerModule* const SequencerModule = FModuleManager::GetModulePtr<ISequencerModule>(TEXT("Sequencer")))
 		{
 			SequencerModule->GetSidebarExtensibilityManager()->RemoveExtender(InToolInstance.SidebarExtender);
 		}
+
+		InToolInstance.SidebarExtender.Reset();
 	}
 }
 
@@ -230,7 +241,6 @@ void FEaseCurveToolExtender::AddSidebarWidget(FMenuBuilder& MenuBuilder, const F
 		, /*bInSearchable=*/true
 		, FText::GetEmpty()
 		, FSlateIcon()
-		//, TAttribute<EVisibility>::CreateStatic(&FEaseCurveToolExtender::GetSidebarVisibility, InToolId)
 	);
 }
 
diff --git a/Engine/Plugins/Experimental/EaseCurveTool/Source/EaseCurveTool/Private/EaseCurveToolModule.cpp b/Engine/Plugins/Experimental/EaseCurveTool/Source/EaseCurveTool/Private/EaseCurveToolModule.cpp
--- a/Engine/Plugins/Experimental/EaseCurveTool/Source/EaseCurveTool/Private/EaseCurveToolModule.cpp
+++ b/Engine/Plugins/Experimental/EaseCurveTool/Source/EaseCurveTool/Private/EaseCurveToolModule.cpp
@@ -34,8 +34,13 @@ void FEaseCurveToolModule::ShutdownModule()
 
 void FEaseCurveToolModule::RegisterContentBrowserExtender()
 {
-	FContentBrowserModule& ContentBrowserModule = FModuleManager::LoadModuleChecked<FContentBrowserModule>(TEXT("ContentBrowser"));
-	TArray<FContentBrowserMenuExtender_SelectedAssets>& AssetMenuExtenders = ContentBrowserModule.GetAllAssetViewContextMenuExtenders();
+	FContentBrowserModule* const ContentBrowserModule = FModuleManager::GetModulePtr<FContentBrowserModule>(TEXT("ContentBrowser"));
+	if (!ContentBrowserModule)
+	{
+		return;
+	}
+
+	TArray<FContentBrowserMenuExtender_SelectedAssets>& AssetMenuExtenders = ContentBrowserModule->GetAllAssetViewContextMenuExtenders();
 
 	AssetMenuExtenders.Add(FContentBrowserMenuExtender_SelectedAssets::CreateStatic(&UE::EaseCurveTool::FEaseCurveLibraryMenu::OnExtendContentBrowserAssetSelectionMenu));
 	ContentBrowserExtenderHandle = AssetMenuExtenders.Last().GetHandle();
@@ -43,19 +48,20 @@ void FEaseCurveToolModule::RegisterContentBrowserExtender()
 
 void FEaseCurveToolModule::UnregisterContentBrowserExtender()
 {
-	if (!ContentBrowserExtenderHandle.IsValid()
-		|| !FModuleManager::Get().IsModuleLoaded(TEXT("ContentBrowser")))
+	if (!ContentBrowserExtenderHandle.IsValid())
 	{
 		return;
 	}
 
-	FContentBrowserModule& ContentBrowserModule = FModuleManager::LoadModuleChecked<FContentBrowserModule>(TEXT("ContentBrowser"));
-	TArray<FContentBrowserMenuExtender_SelectedAssets>& AssetMenuExtenders = ContentBrowserModule.GetAllAssetViewContextMenuExtenders();
+	if (FContentBrowserModule* const ContentBrowserModule = FModuleManager::GetModulePtr<FContentBrowserModule>(TEXT("ContentBrowser")))
+	{
+		TArray<FContentBrowserMenuExtender_SelectedAssets>& AssetMenuExtenders = ContentBrowserModule->GetAllAssetViewContextMenuExtenders();
+		AssetMenuExtenders.RemoveAll([this](const FContentBrowserMenuExtender_SelectedAssets& InDelegate)
+			{
+				return InDelegate.GetHandle() == ContentBrowserExtenderHandle;
+			});
+	}
 
-	AssetMenuExtenders.RemoveAll([this](const FContentBrowserMenuExtender_SelectedAssets& InDelegate)
-		{
-			return InDelegate.GetHandle() == ContentBrowserExtenderHandle;
-		});
 	ContentBrowserExtenderHandle.Reset();
 }
 
@@ -63,10 +69,14 @@ void FEaseCurveToolModule::RegisterCurveEditorExtender()
 {
 	using namespace UE::EaseCurveTool;
 
-	ICurveEditorModule& CurveEditorModule = FModuleManager::LoadModuleChecked<ICurveEditorModule>(TEXT("CurveEditor"));
+	ICurveEditorModule* const CurveEditorModule = FModuleManager::GetModulePtr<ICurveEditorModule>(TEXT("CurveEditor"));
+	if (!CurveEditorModule)
+	{
+		return;
+	}
 
-	CurveEditorExtenderHandle = CurveEditorModule.RegisterEditorExtension(
-		FOnCreateCurveEditorExtension::CreateLambda([](TWeakPtr<FCurveEditor> InCurveEditor)
+	CurveEditorExtenderHandle = CurveEditorModule->RegisterEditorExtension(
+		FOnCreateCurveEditorExtension::CreateLambda([](const TWeakPtr<FCurveEditor> InCurveEditor)
 		{
 			return MakeShared<FEaseCurveEditorExtension>(InCurveEditor);
 		}));
@@ -74,15 +84,16 @@ void FEaseCurveToolModule::RegisterCurveEditorExtender()
 
 void FEaseCurveToolModule::UnregisterCurveEditorExtender()
 {
-	if (!CurveEditorExtenderHandle.IsValid()
-		|| !FModuleManager::Get().IsModuleLoaded(TEXT("CurveEditor")))
+	if (!CurveEditorExtenderHandle.IsValid())
 	{
 		return;
 	}
 
-	ICurveEditorModule& CurveEditorModule = FModuleManager::LoadModuleChecked<ICurveEditorModule>(TEXT("CurveEditor"));
+	if (ICurveEditorModule* const CurveEditorModule = FModuleManager::GetModulePtr<ICurveEditorModule>(TEXT("CurveEditor")))
+	{
+		CurveEditorModule->UnregisterEditorExtension(CurveEditorExtenderHandle);
+	}
 
-	CurveEditorModule.UnregisterEditorExtension(CurveEditorExtenderHandle);
 	CurveEditorExtenderHandle.Reset();
 }
 
