[EaseCurveTool] Set tangent text box values for the graph editor

diff --git a/Engine/Plugins/Experimental/EaseCurveTool/Source/EaseCurveTool/Private/Widgets/SEaseCurveTangents.cpp b/Engine/Plugins/Experimental/EaseCurveTool/Source/EaseCurveTool/Private/Widgets/SEaseCurveTangents.cpp
--- a/Engine/Plugins/Experimental/EaseCurveTool/Source/EaseCurveTool/Private/Widgets/SEaseCurveTangents.cpp
+++ b/Engine/Plugins/Experimental/EaseCurveTool/Source/EaseCurveTool/Private/Widgets/SEaseCurveTangents.cpp
@@ -191,6 +191,11 @@ double SEaseCurveTangents::GetEndTangentWeight() const
 	return Tangents.EndWeight;
 }
 
+void SEaseCurveTangents::SetTangents(const FEaseCurveTangents& InTangents)
+{
+	Tangents = InTangents;
+}
+
 void SEaseCurveTangents::OnStartTangentSpinBoxChanged(const double InNewValue)
 {
 	Tangents.Start = InNewValue;
diff --git a/Engine/Plugins/Experimental/EaseCurveTool/Source/EaseCurveTool/Private/Widgets/SEaseCurveTangents.h b/Engine/Plugins/Experimental/EaseCurveTool/Source/EaseCurveTool/Private/Widgets/SEaseCurveTangents.h
--- a/Engine/Plugins/Experimental/EaseCurveTool/Source/EaseCurveTool/Private/Widgets/SEaseCurveTangents.h
+++ b/Engine/Plugins/Experimental/EaseCurveTool/Source/EaseCurveTool/Private/Widgets/SEaseCurveTangents.h
@@ -29,6 +29,8 @@ public:
 	double GetEndTangent() const;
 	double GetEndTangentWeight() const;
 
+	void SetTangents(const FEaseCurveTangents& InTangents);
+
 protected:
 	TSharedRef<SWidget> ConstructTangentNumBox(const FText& InLabel
 		, const FText& InToolTip
diff --git a/Engine/Plugins/Experimental/EaseCurveTool/Source/EaseCurveTool/Private/Widgets/SEaseCurveTool.cpp b/Engine/Plugins/Experimental/EaseCurveTool/Source/EaseCurveTool/Private/Widgets/SEaseCurveTool.cpp
--- a/Engine/Plugins/Experimental/EaseCurveTool/Source/EaseCurveTool/Private/Widgets/SEaseCurveTool.cpp
+++ b/Engine/Plugins/Experimental/EaseCurveTool/Source/EaseCurveTool/Private/Widgets/SEaseCurveTool.cpp
@@ -9,7 +9,6 @@
 #include "EaseCurveToolSettings.h"
 #include "Editor.h"
 #include "EngineAnalytics.h"
-#include "SEaseCurveTangents.h"
 #include "Framework/Commands/UICommandList.h"
 #include "Framework/MultiBox/MultiBoxBuilder.h"
 #include "Math/UnrealMathUtility.h"
@@ -18,6 +17,7 @@
 #include "Widgets/SEaseCurveEditor.h"
 #include "Widgets/SEaseCurveLibraryComboBox.h"
 #include "Widgets/SEaseCurvePreset.h"
+#include "Widgets/SEaseCurveTangents.h"
 
 #define LOCTEXT_NAMESPACE "SEaseCurveTool"
 
@@ -75,7 +75,7 @@ void SEaseCurveTool::Construct(const FArguments& InArgs, const TSharedRef<FEaseC
 				.HAlign(HAlign_Fill)
 				.Padding(0.f, 3.f, 0.f, 0.f)
 				[
-					SNew(SEaseCurveTangents)
+					SAssignNew(CurveTangentsWidget, SEaseCurveTangents)
 					.InitialTangents(GetTangents())
 					.OnStartTangentChanged(this, &SEaseCurveTool::OnStartTangentSpinBoxChanged)
 					.OnStartWeightChanged(this, &SEaseCurveTool::OnStartTangentWeightSpinBoxChanged)
@@ -250,6 +250,11 @@ void SEaseCurveTool::SetTangents(const FEaseCurveTangents& InTangents, EEaseCurv
 		CurvePresetWidget->ClearSelection();
 	}
 
+	if (CurveTangentsWidget.IsValid())
+	{
+		CurveTangentsWidget->SetTangents(InTangents);
+	}
+
 	// To change the graph UI tangents, we need to change the ease curve object tangents and the graph will reflect.
 	if (bInSetEaseCurve)
 	{
diff --git a/Engine/Plugins/Experimental/EaseCurveTool/Source/EaseCurveTool/Private/Widgets/SEaseCurveTool.h b/Engine/Plugins/Experimental/EaseCurveTool/Source/EaseCurveTool/Private/Widgets/SEaseCurveTool.h
--- a/Engine/Plugins/Experimental/EaseCurveTool/Source/EaseCurveTool/Private/Widgets/SEaseCurveTool.h
+++ b/Engine/Plugins/Experimental/EaseCurveTool/Source/EaseCurveTool/Private/Widgets/SEaseCurveTool.h
@@ -17,6 +17,7 @@ namespace UE::EaseCurveTool
 class FEaseCurveToolContextMenu;
 class SEaseCurveEditor;
 class SEaseCurvePreset;
+class SEaseCurveTangents;
 
 class SEaseCurveTool
 	: public SCompoundWidget
@@ -94,6 +95,7 @@ protected:
 
 	TSharedPtr<SEaseCurvePreset> CurvePresetWidget;
 	TSharedPtr<SEaseCurveEditor> CurveEaseEditorWidget;
+	TSharedPtr<SEaseCurveTangents> CurveTangentsWidget;
 
 	int32 CurrentGraphSize = DefaultGraphSize;
 
