[EaseCurveTool]
- Add FEaseCurvePresetHandle to work with category/name combos now that we are no longer using Json files for categories
- Fixes ease curves with the same name being applied with the first found category instead of the category it should be

diff --git a/Engine/Plugins/Experimental/EaseCurveTool/Source/EaseCurveTool/Private/EaseCurveLibrary.cpp b/Engine/Plugins/Experimental/EaseCurveTool/Source/EaseCurveTool/Private/EaseCurveLibrary.cpp
--- a/Engine/Plugins/Experimental/EaseCurveTool/Source/EaseCurveTool/Private/EaseCurveLibrary.cpp
+++ b/Engine/Plugins/Experimental/EaseCurveTool/Source/EaseCurveTool/Private/EaseCurveLibrary.cpp
@@ -67,17 +67,11 @@ TArray<FEaseCurvePreset> UEaseCurveLibrary::GetDefaultPresets()
 	return DefaultPresets;
 }
 
-bool UEaseCurveLibrary::DoesPresetExist(const FEaseCurvePreset& InPreset) const
+bool UEaseCurveLibrary::DoesPresetExist(const FEaseCurvePresetHandle& InPresetHandle) const
 {
-	return DoesPresetExist(InPreset.Category, InPreset.Name);
-}
-
-bool UEaseCurveLibrary::DoesPresetExist(const FText& InCategory, const FText& InName) const
-{
-	return Presets.ContainsByPredicate([&InCategory, InName](const FEaseCurvePreset& InPreset)
+	return Presets.ContainsByPredicate([&InPresetHandle](const FEaseCurvePreset& InPreset)
 		{
-			return InPreset.Category.EqualToCaseIgnored(InCategory)
-				&& InPreset.Name.EqualToCaseIgnored(InName);
+			return InPreset == InPresetHandle;
 		});
 }
 
@@ -134,7 +128,7 @@ TArray<FEaseCurvePreset> UEaseCurveLibrary::GetCategoryPresets(const FText& InCa
 
 bool UEaseCurveLibrary::AddPreset(const FEaseCurvePreset& InPreset)
 {
-	if (DoesPresetExist(InPreset.Category, InPreset.Name))
+	if (DoesPresetExist(InPreset.GetHandle()))
 	{
 		UE_LOG(LogEaseCurveLibrary, Warning, TEXT("UEaseCurveLibrary::AddPreset() - Preset already exists: %s")
 			, *InPreset.Name.ToString());
@@ -166,13 +160,12 @@ bool UEaseCurveLibrary::AddPresetToNewCategory(const FText& InName, const FEaseC
 	return AddPreset(OutNewPreset);
 }
 
-bool UEaseCurveLibrary::RemovePreset(const FEaseCurvePreset& InPreset)
+bool UEaseCurveLibrary::RemovePreset(const FEaseCurvePresetHandle& InPresetHandle)
 {
-	const int32 NumRemoved = Presets.RemoveAll([&InPreset]
+	const int32 NumRemoved = Presets.RemoveAll([&InPresetHandle]
 		(const FEaseCurvePreset& InThisPreset)
 		{
-			return InThisPreset.Category.EqualToCaseIgnored(InPreset.Category)
-				&& InThisPreset.Name.EqualToCaseIgnored(InPreset.Name);
+			return InThisPreset == InPresetHandle;
 		});
 
 	const bool bRemoved = NumRemoved > 0;
@@ -208,13 +201,12 @@ bool UEaseCurveLibrary::ChangePresetCategory(const FEaseCurvePreset& InPreset, c
 	return true;
 }
 
-bool UEaseCurveLibrary::FindPreset(const FText& InCategory, const FText& InName, FEaseCurvePreset& OutPreset)
+bool UEaseCurveLibrary::FindPreset(const FEaseCurvePresetHandle& InPresetHandle, FEaseCurvePreset& OutPreset)
 {
-	const FEaseCurvePreset* FoundPreset = Presets.FindByPredicate([&InCategory, &InName]
+	const FEaseCurvePreset* FoundPreset = Presets.FindByPredicate([&InPresetHandle]
 		(const FEaseCurvePreset& InThisPreset)
 		{
-			return InThisPreset.Category.EqualToCaseIgnored(InCategory)
-				&& InThisPreset.Name.EqualToCaseIgnored(InName);
+			return InThisPreset == InPresetHandle;
 		});
 
 	if (FoundPreset)
@@ -341,15 +333,15 @@ bool UEaseCurveLibrary::RemovePresetCategory(const FText& InCategory)
 	return bRemoved;
 }
 
-bool UEaseCurveLibrary::RenamePreset(const FText& InCategory, const FText& InPresetName, const FText& InNewPresetName)
+bool UEaseCurveLibrary::RenamePreset(const FEaseCurvePresetHandle& InPresetHandle, const FText& InNewPresetName)
 {
-	if (DoesPresetExist(InCategory, InNewPresetName))
+	if (DoesPresetExist(FEaseCurvePresetHandle(InPresetHandle.Category, InNewPresetName)))
 	{
 		return false;
 	}
 
 	FEaseCurvePreset FoundPreset;
-	if (!FindPreset(InCategory, InPresetName, FoundPreset))
+	if (!FindPreset(InPresetHandle, FoundPreset))
 	{
 		return false;
 	}
diff --git a/Engine/Plugins/Experimental/EaseCurveTool/Source/EaseCurveTool/Private/EaseCurvePreset.h b/Engine/Plugins/Experimental/EaseCurveTool/Source/EaseCurveTool/Private/EaseCurvePreset.h
--- a/Engine/Plugins/Experimental/EaseCurveTool/Source/EaseCurveTool/Private/EaseCurvePreset.h
+++ b/Engine/Plugins/Experimental/EaseCurveTool/Source/EaseCurveTool/Private/EaseCurvePreset.h
@@ -7,6 +7,26 @@
 #include "Templates/SharedPointer.h"
 #include "EaseCurvePreset.generated.h"
 
+/** Store category and name as a handle to pass around for preset identification */
+struct FEaseCurvePresetHandle
+{
+	FEaseCurvePresetHandle() {}
+	FEaseCurvePresetHandle(const FText& InCategory, const FText& InName)
+		: Category(InCategory), Name(InName)
+	{
+		ensureMsgf(Category.IsEmpty() == false, TEXT("Category cannot be empty"));
+		ensureMsgf(Name.IsEmpty() == false, TEXT("Name cannot be empty"));
+	}
+
+	FORCEINLINE bool operator==(const FEaseCurvePresetHandle& InRhs) const
+	{
+		return Name.EqualToCaseIgnored(InRhs.Name) && Category.EqualToCaseIgnored(InRhs.Category);
+	}
+
+	FText Category;
+	FText Name;
+};
+
 USTRUCT()
 struct FEaseCurvePreset
 {
@@ -16,16 +36,23 @@ struct FEaseCurvePreset
 	FEaseCurvePreset(const FText& InName, const FText& InCategory, const FEaseCurveTangents& InTangents)
 		: Name(InName), Category(InCategory), Tangents(InTangents)
 	{}
+	FEaseCurvePreset(const FEaseCurvePresetHandle& InHandle, const FEaseCurveTangents& InTangents)
+		: Name(InHandle.Name), Category(InHandle.Category), Tangents(InTangents)
+	{}
 
+	FORCEINLINE bool operator==(const FEaseCurvePresetHandle& InRhs) const
+	{
+		return GetHandle() == InRhs;
+	}
 	FORCEINLINE bool operator==(const FEaseCurvePreset& InRhs) const
 	{
-		return Name.EqualToCaseIgnored(InRhs.Name)
-			&& Category.EqualToCaseIgnored(InRhs.Category);
+		return GetHandle() == InRhs.GetHandle();
 	}
 	FORCEINLINE bool operator!=(const FEaseCurvePreset& InRhs) const
 	{
 		return !(*this == InRhs);
 	}
+
 	FORCEINLINE bool operator<(const FEaseCurvePreset& InRhs) const
 	{
 		// Order from soft to hard based on curve length
@@ -34,6 +61,11 @@ struct FEaseCurvePreset
 			&& Tangents.CalculateCurveLength() < InRhs.Tangents.CalculateCurveLength();
 	}
 
+	FEaseCurvePresetHandle GetHandle() const
+	{
+		return FEaseCurvePresetHandle(Name, Category);
+	}
+
 	UPROPERTY(Config, EditAnywhere, Category = "EaseCurvePreset")
 	FText Name;
 
diff --git a/Engine/Plugins/Experimental/EaseCurveTool/Source/EaseCurveTool/Private/Widgets/SEaseCurvePreset.cpp b/Engine/Plugins/Experimental/EaseCurveTool/Source/EaseCurveTool/Private/Widgets/SEaseCurvePreset.cpp
--- a/Engine/Plugins/Experimental/EaseCurveTool/Source/EaseCurveTool/Private/Widgets/SEaseCurvePreset.cpp
+++ b/Engine/Plugins/Experimental/EaseCurveTool/Source/EaseCurveTool/Private/Widgets/SEaseCurvePreset.cpp
@@ -143,7 +143,7 @@ void SEaseCurvePreset::HandlePresetChanged()
 
 	if (bHasSelectedItem)
 	{
-		PresetComboBox->SetSelectedItem(Preset);
+		PresetComboBox->SetSelectedItem(Preset.GetHandle());
 	}
 }
 
@@ -190,7 +190,7 @@ void SEaseCurvePreset::OnNewPresetTextCommitted(const FText& InNewText, ETextCom
 				if (PresetLibrary->AddPresetToNewCategory(InNewText, NewTangents, NewPreset))
 				{
 					PresetComboBox->Reload();
-					PresetComboBox->SetSelectedItem(NewPreset);
+					PresetComboBox->SetSelectedItem(NewPreset.GetHandle());
 				}
 			}
 		}
@@ -205,9 +205,9 @@ void SEaseCurvePreset::ClearSelection()
 	PresetComboBox->ClearSelection();
 }
 
-bool SEaseCurvePreset::SetSelectedItem(const FText& InName)
+bool SEaseCurvePreset::SetSelectedItem(const FEaseCurvePresetHandle& InPresetHandle)
 {
-	return PresetComboBox->SetSelectedItem(InName);
+	return PresetComboBox->SetSelectedItem(InPresetHandle);
 }
 
 bool SEaseCurvePreset::SetSelectedItem(const FEaseCurveTangents& InTangents)
@@ -215,11 +215,6 @@ bool SEaseCurvePreset::SetSelectedItem(const FEaseCurveTangents& InTangents)
 	return PresetComboBox->SetSelectedItem(InTangents);
 }
 
-bool SEaseCurvePreset::SetSelectedItem(const FEaseCurvePreset& InPreset)
-{
-	return SetSelectedItem(InPreset.Name);
-}
-
 } // namespace UE::EaseCurveTool
 
 #undef LOCTEXT_NAMESPACE
diff --git a/Engine/Plugins/Experimental/EaseCurveTool/Source/EaseCurveTool/Private/Widgets/SEaseCurvePreset.h b/Engine/Plugins/Experimental/EaseCurveTool/Source/EaseCurveTool/Private/Widgets/SEaseCurvePreset.h
--- a/Engine/Plugins/Experimental/EaseCurveTool/Source/EaseCurveTool/Private/Widgets/SEaseCurvePreset.h
+++ b/Engine/Plugins/Experimental/EaseCurveTool/Source/EaseCurveTool/Private/Widgets/SEaseCurvePreset.h
@@ -37,9 +37,8 @@ public:
 
 	void ClearSelection();
 	
-	bool SetSelectedItem(const FText& InName);
+	bool SetSelectedItem(const FEaseCurvePresetHandle& InPresetHandle);
 	bool SetSelectedItem(const FEaseCurveTangents& InTangents);
-	bool SetSelectedItem(const FEaseCurvePreset& InPreset);
 
 protected:
 	void HandleLibraryChanged(const TWeakObjectPtr<UEaseCurveLibrary> InWeakLibrary);
diff --git a/Engine/Plugins/Experimental/EaseCurveTool/Source/EaseCurveTool/Private/Widgets/SEaseCurvePresetComboBox.cpp b/Engine/Plugins/Experimental/EaseCurveTool/Source/EaseCurveTool/Private/Widgets/SEaseCurvePresetComboBox.cpp
--- a/Engine/Plugins/Experimental/EaseCurveTool/Source/EaseCurveTool/Private/Widgets/SEaseCurvePresetComboBox.cpp
+++ b/Engine/Plugins/Experimental/EaseCurveTool/Source/EaseCurveTool/Private/Widgets/SEaseCurvePresetComboBox.cpp
@@ -179,20 +179,20 @@ bool SEaseCurvePresetComboBox::GetSelectedItem(FEaseCurvePreset& OutPreset) cons
 	return true;
 }
 
-bool SEaseCurvePresetComboBox::SetSelectedItem(const FText& InName)
+bool SEaseCurvePresetComboBox::SetSelectedItem(const FEaseCurvePresetHandle& InPresetHandle)
 {
-	const TSharedPtr<FEaseCurveTool> Tool = WeakTool.Pin();
-	if (!Tool.IsValid())
+	if (!PresetList.IsValid())
 	{
 		return false;
 	}
 
-	if (!PresetList.IsValid())
+	const TSharedPtr<FEaseCurveTool> Tool = WeakTool.Pin();
+	if (!Tool.IsValid())
 	{
 		return false;
 	}
 
-	const TSharedPtr<FEaseCurvePreset> FoundItem = PresetList->FindItemByName(InName);
+	const TSharedPtr<FEaseCurvePreset> FoundItem = PresetList->FindItem(InPresetHandle);
 	if (!FoundItem.IsValid())
 	{
 		return false;
@@ -227,13 +227,8 @@ bool SEaseCurvePresetComboBox::SetSelectedItem(const FEaseCurveTangents& InTange
 	SelectedItem = FoundItem;
 
 	GenerateSelectedRowWidget();
-	
-	return true;
-}
 
-bool SEaseCurvePresetComboBox::SetSelectedItem(const FEaseCurvePreset& InPreset)
-{
-	return SetSelectedItem(InPreset.Name);
+	return true;
 }
 
 void SEaseCurvePresetComboBox::Reload()
diff --git a/Engine/Plugins/Experimental/EaseCurveTool/Source/EaseCurveTool/Private/Widgets/SEaseCurvePresetComboBox.h b/Engine/Plugins/Experimental/EaseCurveTool/Source/EaseCurveTool/Private/Widgets/SEaseCurvePresetComboBox.h
--- a/Engine/Plugins/Experimental/EaseCurveTool/Source/EaseCurveTool/Private/Widgets/SEaseCurvePresetComboBox.h
+++ b/Engine/Plugins/Experimental/EaseCurveTool/Source/EaseCurveTool/Private/Widgets/SEaseCurvePresetComboBox.h
@@ -42,9 +42,8 @@ public:
 	
 	bool GetSelectedItem(FEaseCurvePreset& OutPreset) const;
 	
-	bool SetSelectedItem(const FText& InName);
+	bool SetSelectedItem(const FEaseCurvePresetHandle& InPresetHandle);
 	bool SetSelectedItem(const FEaseCurveTangents& InTangents);
-	bool SetSelectedItem(const FEaseCurvePreset& InPreset);
 
 	void Reload();
 
diff --git a/Engine/Plugins/Experimental/EaseCurveTool/Source/EaseCurveTool/Private/Widgets/SEaseCurvePresetList.cpp b/Engine/Plugins/Experimental/EaseCurveTool/Source/EaseCurveTool/Private/Widgets/SEaseCurvePresetList.cpp
--- a/Engine/Plugins/Experimental/EaseCurveTool/Source/EaseCurveTool/Private/Widgets/SEaseCurvePresetList.cpp
+++ b/Engine/Plugins/Experimental/EaseCurveTool/Source/EaseCurveTool/Private/Widgets/SEaseCurvePresetList.cpp
@@ -313,12 +313,12 @@ bool SEaseCurvePresetList::GetSelectedItem(FEaseCurvePreset& OutPreset) const
 	return true;
 }
 
-TSharedPtr<FEaseCurvePreset> SEaseCurvePresetList::FindItemByName(const FText& InPresetName) const
+TSharedPtr<FEaseCurvePreset> SEaseCurvePresetList::FindItem(const FEaseCurvePresetHandle& InPresetHandle) const
 {
-	const TSharedPtr<FEaseCurvePreset>* FoundPresetItem = PresetItems.FindByPredicate([&InPresetName]
+	const TSharedPtr<FEaseCurvePreset>* FoundPresetItem = PresetItems.FindByPredicate([&InPresetHandle]
 		(const TSharedPtr<FEaseCurvePreset>& InThisPreset)
 		{
-			return InThisPreset.IsValid() && InThisPreset->Name.EqualToCaseIgnored(InPresetName);
+			return InThisPreset.IsValid() && *InThisPreset == InPresetHandle;
 		});
 
 	if (FoundPresetItem)
@@ -351,9 +351,9 @@ void SEaseCurvePresetList::ClearSelection()
 	OnPresetChanged.ExecuteIfBound(nullptr);
 }
 
-bool SEaseCurvePresetList::SetSelectedItem(const FText& InName)
+bool SEaseCurvePresetList::SetSelectedItem(const FEaseCurvePresetHandle& InPresetHandle)
 {
-	const TSharedPtr<FEaseCurvePreset> FoundItem = FindItemByName(InName);
+	const TSharedPtr<FEaseCurvePreset> FoundItem = FindItem(InPresetHandle);
 	if (!FoundItem.IsValid())
 	{
 		return false;
@@ -379,11 +379,6 @@ bool SEaseCurvePresetList::SetSelectedItem(const FEaseCurveTangents& InTangents)
 	return true;
 }
 
-bool SEaseCurvePresetList::SetSelectedItem(const FEaseCurvePreset& InPreset)
-{
-	return SetSelectedItem(InPreset.Name);
-}
-
 bool SEaseCurvePresetList::IsInEditMode() const
 {
 	return bEditMode.Get(false);
@@ -455,7 +450,7 @@ bool SEaseCurvePresetList::HandlePresetDelete(const TSharedPtr<FEaseCurvePreset>
 		return false;
 	}
 
-	if (!PresetLibrary->RemovePreset(*InPreset))
+	if (!PresetLibrary->RemovePreset(InPreset->GetHandle()))
 	{
 		return false;
 	}
@@ -478,7 +473,7 @@ bool SEaseCurvePresetList::HandlePresetRename(const TSharedPtr<FEaseCurvePreset>
 		return false;
 	}
 
-	if (!PresetLibrary->RenamePreset(InPreset->Category, InPreset->Name, InNewName))
+	if (!PresetLibrary->RenamePreset(InPreset->GetHandle(), InNewName))
 	{
 		return false;
 	}
@@ -539,9 +534,12 @@ bool SEaseCurvePresetList::HandleEndPresetMove(const TSharedPtr<FEaseCurvePreset
 
 bool SEaseCurvePresetList::HandlePresetClick(const TSharedPtr<FEaseCurvePreset>& InPreset)
 {
-	SetSelectedItem(*InPreset);
-
-	return true;
+	if (InPreset.IsValid())
+	{
+		SetSelectedItem(InPreset->GetHandle());
+		return true;
+	}
+	return false;
 }
 
 bool SEaseCurvePresetList::HandleSetQuickEase(const TSharedPtr<FEaseCurvePreset>& InPreset)
@@ -568,7 +566,7 @@ FReply SEaseCurvePresetList::OnDeletePresetClick()
 		return FReply::Handled();
 	}
 
-	if (!PresetLibrary->RemovePreset(*SelectedItem))
+	if (!PresetLibrary->RemovePreset(SelectedItem->GetHandle()))
 	{
 		return FReply::Handled();
 	}
diff --git a/Engine/Plugins/Experimental/EaseCurveTool/Source/EaseCurveTool/Private/Widgets/SEaseCurvePresetList.h b/Engine/Plugins/Experimental/EaseCurveTool/Source/EaseCurveTool/Private/Widgets/SEaseCurvePresetList.h
--- a/Engine/Plugins/Experimental/EaseCurveTool/Source/EaseCurveTool/Private/Widgets/SEaseCurvePresetList.h
+++ b/Engine/Plugins/Experimental/EaseCurveTool/Source/EaseCurveTool/Private/Widgets/SEaseCurvePresetList.h
@@ -8,6 +8,7 @@
 #include "Widgets/DeclarativeSyntaxSupport.h"
 #include "Widgets/SCompoundWidget.h"
 
+struct FEaseCurvePresetHandle;
 class SBox;
 class SUniformWrapPanel;
 class UEaseCurveLibrary;
@@ -46,11 +47,10 @@ public:
 
 	void ClearSelection();
 
-	bool SetSelectedItem(const FText& InName);
+	bool SetSelectedItem(const FEaseCurvePresetHandle& InPresetHandle);
 	bool SetSelectedItem(const FEaseCurveTangents& InTangents);
-	bool SetSelectedItem(const FEaseCurvePreset& InPreset);
 
-	TSharedPtr<FEaseCurvePreset> FindItemByName(const FText& InPresetName) const;
+	TSharedPtr<FEaseCurvePreset> FindItem(const FEaseCurvePresetHandle& InPresetHandle) const;
 	TSharedPtr<FEaseCurvePreset> FindItemByTangents(const FEaseCurveTangents& InTangents, const double InErrorTolerance = 0.01) const;
 
 	bool IsInEditMode() const;
diff --git a/Engine/Plugins/Experimental/EaseCurveTool/Source/EaseCurveTool/Public/EaseCurveLibrary.h b/Engine/Plugins/Experimental/EaseCurveTool/Source/EaseCurveTool/Public/EaseCurveLibrary.h
--- a/Engine/Plugins/Experimental/EaseCurveTool/Source/EaseCurveTool/Public/EaseCurveLibrary.h
+++ b/Engine/Plugins/Experimental/EaseCurveTool/Source/EaseCurveTool/Public/EaseCurveLibrary.h
@@ -26,8 +26,7 @@ public:
 
 	FSimpleMulticastDelegate& OnPresetChanged() { return PresetChangedDelegate; }
 
-	bool DoesPresetExist(const FEaseCurvePreset& InPreset) const;
-	bool DoesPresetExist(const FText& InCategory, const FText& InName) const;
+	bool DoesPresetExist(const FEaseCurvePresetHandle& InPresetHandle) const;
 
 	const TArray<FEaseCurvePreset>& GetPresets() const;
 
@@ -38,11 +37,11 @@ public:
 	bool AddPreset(const FEaseCurvePreset& InPreset);
 	bool AddPresetToNewCategory(const FText& InName, const FEaseCurveTangents& InTangents, FEaseCurvePreset& OutNewPreset);
 
-	bool RemovePreset(const FEaseCurvePreset& InPreset);
+	bool RemovePreset(const FEaseCurvePresetHandle& InPresetHandle);
 
 	bool ChangePresetCategory(const FEaseCurvePreset& InPreset, const FText& InNewCategory);
 
-	bool FindPreset(const FText& InCategory, const FText& InName, FEaseCurvePreset& OutPreset);
+	bool FindPreset(const FEaseCurvePresetHandle& InPresetHandle, FEaseCurvePreset& OutPreset);
 	bool FindPresetByTangents(const FEaseCurveTangents& InTangents
 		, FEaseCurvePreset& OutPreset, const double InErrorTolerance = 0.01);
 
@@ -50,7 +49,7 @@ public:
 	bool RenamePresetCategory(const FText& InCategory, const FText& InNewCategory);
 	bool RemovePresetCategory(const FText& InCategory);
 
-	bool RenamePreset(const FText& InCategory, const FText& InPresetName, const FText& InNewPresetName);
+	bool RenamePreset(const FEaseCurvePresetHandle& InPresetHandle, const FText& InNewPresetName);
 
 	void SetToDefaultPresets();
 
