[EaseCurveTool] Change parameters for category and name so GetHandle() returns correctly

diff --git a/Engine/Plugins/Experimental/EaseCurveTool/Source/EaseCurveTool/Private/EaseCurvePreset.h b/Engine/Plugins/Experimental/EaseCurveTool/Source/EaseCurveTool/Private/EaseCurvePreset.h
--- a/Engine/Plugins/Experimental/EaseCurveTool/Source/EaseCurveTool/Private/EaseCurvePreset.h
+++ b/Engine/Plugins/Experimental/EaseCurveTool/Source/EaseCurveTool/Private/EaseCurvePreset.h
@@ -63,7 +63,7 @@ struct FEaseCurvePreset
 
 	FEaseCurvePresetHandle GetHandle() const
 	{
-		return FEaseCurvePresetHandle(Name, Category);
+		return FEaseCurvePresetHandle(Category, Name);
 	}
 
 	UPROPERTY(Config, EditAnywhere, Category = "EaseCurvePreset")
