[EaseCurveTool] - Fill available size when undocked
- Fix crash on editor startup when attempting to spawn the tool tab when the tab was last undocked on editor close
- Maximize available space when tool is undocked
- Put makeup on graph editor border
- Fix RenamePreset() from using local copy to set name

diff --git a/Engine/Plugins/Experimental/EaseCurveTool/Source/EaseCurveTool/Private/EaseCurveLibrary.cpp b/Engine/Plugins/Experimental/EaseCurveTool/Source/EaseCurveTool/Private/EaseCurveLibrary.cpp
--- a/Engine/Plugins/Experimental/EaseCurveTool/Source/EaseCurveTool/Private/EaseCurveLibrary.cpp
+++ b/Engine/Plugins/Experimental/EaseCurveTool/Source/EaseCurveTool/Private/EaseCurveLibrary.cpp
@@ -203,7 +203,7 @@ bool UEaseCurveLibrary::ChangePresetCategory(const FEaseCurvePreset& InPreset, c
 
 bool UEaseCurveLibrary::FindPreset(const FEaseCurvePresetHandle& InPresetHandle, FEaseCurvePreset& OutPreset)
 {
-	const FEaseCurvePreset* FoundPreset = Presets.FindByPredicate([&InPresetHandle]
+	const FEaseCurvePreset* const FoundPreset = Presets.FindByPredicate([&InPresetHandle]
 		(const FEaseCurvePreset& InThisPreset)
 		{
 			return InThisPreset == InPresetHandle;
@@ -222,7 +222,7 @@ bool UEaseCurveLibrary::FindPreset(const FEaseCurvePresetHandle& InPresetHandle,
 bool UEaseCurveLibrary::FindPresetByTangents(const FEaseCurveTangents& InTangents
 	, FEaseCurvePreset& OutPreset, const double InErrorTolerance)
 {
-	const FEaseCurvePreset* FoundPreset = Presets.FindByPredicate([&InTangents, InErrorTolerance]
+	const FEaseCurvePreset* const FoundPreset = Presets.FindByPredicate([&InTangents, InErrorTolerance]
 		(const FEaseCurvePreset& InThisPreset)
 		{
 			return InThisPreset.Tangents.IsNearlyEqual(InTangents, InErrorTolerance);
@@ -335,18 +335,17 @@ bool UEaseCurveLibrary::RemovePresetCategory(const FText& InCategory)
 
 bool UEaseCurveLibrary::RenamePreset(const FEaseCurvePresetHandle& InPresetHandle, const FText& InNewPresetName)
 {
-	if (DoesPresetExist(FEaseCurvePresetHandle(InPresetHandle.Category, InNewPresetName)))
-	{
-		return false;
-	}
-
-	FEaseCurvePreset FoundPreset;
-	if (!FindPreset(InPresetHandle, FoundPreset))
+	FEaseCurvePreset* const FoundPreset = Presets.FindByPredicate([&InPresetHandle]
+		(const FEaseCurvePreset& InThisPreset)
+		{
+			return InThisPreset == InPresetHandle;
+		});
+	if (!FoundPreset)
 	{
 		return false;
 	}
 
-	FoundPreset.Name = InNewPresetName;
+	FoundPreset->Name = InNewPresetName;
 
 	PresetChangedDelegate.Broadcast();
 
diff --git a/Engine/Plugins/Experimental/EaseCurveTool/Source/EaseCurveTool/Private/EaseCurveStyle.cpp b/Engine/Plugins/Experimental/EaseCurveTool/Source/EaseCurveTool/Private/EaseCurveStyle.cpp
--- a/Engine/Plugins/Experimental/EaseCurveTool/Source/EaseCurveTool/Private/EaseCurveStyle.cpp
+++ b/Engine/Plugins/Experimental/EaseCurveTool/Source/EaseCurveTool/Private/EaseCurveStyle.cpp
@@ -69,6 +69,9 @@ FEaseCurveStyle::FEaseCurveStyle()
 		.SetPadding(ToolButtonPadding);
 	Set("ToolToggleButton", ToolToggleButtonStyle);
 
+	Set("Editor.Border", new FSlateRoundedBoxBrush(FStyleColors::Transparent, 4.f, FStyleColors::Dropdown, 1.f));
+	Set("Editor.Border.Hover", new FSlateRoundedBoxBrush(FStyleColors::Transparent, 4.f, FStyleColors::DropdownOutline, 1.f));
+
 	Set("Editor.LabelFont", FSlateFontInfo(FCoreStyle::GetDefaultFont(), 7, TEXT("Regular")));
 	Set("Editor.ErrorFont", FSlateFontInfo(FCoreStyle::GetDefaultFont(), 9, TEXT("Italic")));
 
diff --git a/Engine/Plugins/Experimental/EaseCurveTool/Source/EaseCurveTool/Private/EaseCurveToolExtender.cpp b/Engine/Plugins/Experimental/EaseCurveTool/Source/EaseCurveTool/Private/EaseCurveToolExtender.cpp
--- a/Engine/Plugins/Experimental/EaseCurveTool/Source/EaseCurveTool/Private/EaseCurveToolExtender.cpp
+++ b/Engine/Plugins/Experimental/EaseCurveTool/Source/EaseCurveTool/Private/EaseCurveToolExtender.cpp
@@ -24,6 +24,8 @@
 namespace UE::EaseCurveTool
 {
 
+bool FEaseCurveToolExtender::bFirstTimeSequencerCreated = false;
+
 FEaseCurveToolExtender& FEaseCurveToolExtender::Get()
 {
 	static FEaseCurveToolExtender Instance;
@@ -66,7 +68,6 @@ FEaseCurveToolExtender::FEaseCurveToolInstance* FEaseCurveToolExtender::FindOrAd
 	}
 
 	FEaseCurveToolInstance& NewToolInstance = ToolInstances.FindOrAdd(NewToolId);
-
 	NewToolInstance.ToolId = NewToolId;
 	NewToolInstance.Instance = MakeShared<FEaseCurveTool>(InSequencer);
 	NewToolInstance.WeakSequencer = InSequencer;
@@ -84,6 +85,20 @@ void FEaseCurveToolExtender::OnSequencerCreated(const TSharedRef<ISequencer> InS
 	{
 		return;
 	}
+
+	// Since the ease curve tool instance doesn't get created until the sequence is created,
+	// we check and show the tool tab here after we know the instance exists. Otherwise we will
+	// crash on editor startup attempting to bind commands through the constructor.
+	if (!bFirstTimeSequencerCreated)
+	{
+		const UEaseCurveToolSettings* const ToolSettings = GetDefault<UEaseCurveToolSettings>();
+		if (ToolSettings && ToolSettings->IsToolVisible())
+		{
+			NewToolInstance->Instance->ShowHideToolTab(true);
+		}
+	}
+
+	bFirstTimeSequencerCreated = true;
 }
 
 void FEaseCurveToolExtender::OnSequencerClosed(const TSharedRef<ISequencer> InSequencer, const FName InToolId)
diff --git a/Engine/Plugins/Experimental/EaseCurveTool/Source/EaseCurveTool/Private/EaseCurveToolExtender.h b/Engine/Plugins/Experimental/EaseCurveTool/Source/EaseCurveTool/Private/EaseCurveToolExtender.h
--- a/Engine/Plugins/Experimental/EaseCurveTool/Source/EaseCurveTool/Private/EaseCurveToolExtender.h
+++ b/Engine/Plugins/Experimental/EaseCurveTool/Source/EaseCurveTool/Private/EaseCurveToolExtender.h
@@ -80,6 +80,8 @@ protected:
 	static EVisibility GetSidebarVisibility(const FName InToolId);
 
 private:
+	static bool bFirstTimeSequencerCreated;
+
 	FDelegateHandle SequencerCreatedHandle;
 
 	/** Registered Sequencer types to Ease Curve Tool instances */
diff --git a/Engine/Plugins/Experimental/EaseCurveTool/Source/EaseCurveTool/Private/EaseCurveToolTab.cpp b/Engine/Plugins/Experimental/EaseCurveTool/Source/EaseCurveTool/Private/EaseCurveToolTab.cpp
--- a/Engine/Plugins/Experimental/EaseCurveTool/Source/EaseCurveTool/Private/EaseCurveToolTab.cpp
+++ b/Engine/Plugins/Experimental/EaseCurveTool/Source/EaseCurveTool/Private/EaseCurveToolTab.cpp
@@ -4,6 +4,7 @@
 #include "EaseCurveStyle.h"
 #include "EaseCurveTool.h"
 #include "EaseCurveToolSettings.h"
+#include "Editor.h"
 #include "Framework/Docking/TabManager.h"
 #include "ISequencer.h"
 #include "Toolkits/IToolkitHost.h"
@@ -25,15 +26,6 @@ void FEaseCurveToolTab::Init()
 	bShuttingDown = false;
 
 	RegisterToolTab();
-
-	const UEaseCurveToolSettings* const ToolSettings = GetDefault<UEaseCurveToolSettings>();
-    if (ToolSettings && ToolSettings->IsToolVisible())
-    {
-    	if (const TSharedPtr<FTabManager> TabManager = GetTabManager())
-    	{
-    		TabManager->TryInvokeTab(ToolTabId);
-    	}
-    }
 }
 
 void FEaseCurveToolTab::Shutdown()
@@ -196,8 +188,6 @@ TSharedRef<SDockTab> FEaseCurveToolTab::SpawnToolTab(const FSpawnTabArgs& InArgs
 
 	VisibilityChangedDelegate.Broadcast(true);
 
-	//OwnerTool.Refresh();
-
 	return ToolTab;
 }
 
diff --git a/Engine/Plugins/Experimental/EaseCurveTool/Source/EaseCurveTool/Private/Widgets/SEaseCurveEditor.cpp b/Engine/Plugins/Experimental/EaseCurveTool/Source/EaseCurveTool/Private/Widgets/SEaseCurveEditor.cpp
--- a/Engine/Plugins/Experimental/EaseCurveTool/Source/EaseCurveTool/Private/Widgets/SEaseCurveEditor.cpp
+++ b/Engine/Plugins/Experimental/EaseCurveTool/Source/EaseCurveTool/Private/Widgets/SEaseCurveEditor.cpp
@@ -102,12 +102,6 @@ void SEaseCurveEditor::Construct(const FArguments& InArgs, const TSharedRef<FEas
 	OnDragStart = InArgs._OnDragStart;
 	OnDragEnd = InArgs._OnDragEnd;
 
-	// If editor size is set, use it, otherwise, use default value
-	if (DesiredSize.Get().IsZero())
-	{
-		DesiredSize.Set(FVector2D(128));
-	}
-
 	if (GEditor)
 	{
 		GEditor->RegisterForUndo(this);
@@ -146,7 +140,8 @@ SEaseCurveEditor::~SEaseCurveEditor()
 
 FVector2D SEaseCurveEditor::ComputeDesiredSize(const float InLayoutScaleMultiplier) const
 {
-	return DesiredSize.Get();
+	const TOptional<FVector2D> SizeFromAttribute = DesiredSize.Get(TOptional<FVector2D>());
+	return SizeFromAttribute.Get(FVector2D::ZeroVector);
 }
 
 int32 SEaseCurveEditor::OnPaint(const FPaintArgs& InArgs
diff --git a/Engine/Plugins/Experimental/EaseCurveTool/Source/EaseCurveTool/Private/Widgets/SEaseCurveEditor.h b/Engine/Plugins/Experimental/EaseCurveTool/Source/EaseCurveTool/Private/Widgets/SEaseCurveEditor.h
--- a/Engine/Plugins/Experimental/EaseCurveTool/Source/EaseCurveTool/Private/Widgets/SEaseCurveEditor.h
+++ b/Engine/Plugins/Experimental/EaseCurveTool/Source/EaseCurveTool/Private/Widgets/SEaseCurveEditor.h
@@ -31,8 +31,7 @@ class SEaseCurveEditor :
 {
 public:
 	SLATE_BEGIN_ARGS(SEaseCurveEditor)
-		: _DesiredSize(FVector2D(200.f))
-		, _DisplayRate(FFrameRate(30, 1))
+		: _DisplayRate(FFrameRate(30, 1))
 		, _ShowInputGridNumbers(true)
 		, _ShowOutputGridNumbers(true)
 		, _NormalAreaColor(FLinearColor(1.0f, 1.0f, 1.0f, 0.05f))
@@ -48,7 +47,7 @@ public:
 		SLATE_ATTRIBUTE(TOptional<float>, DataMinInput)
 		SLATE_ATTRIBUTE(TOptional<float>, DataMaxInput)
 
-		SLATE_ATTRIBUTE(FVector2D, DesiredSize)
+		SLATE_ATTRIBUTE(TOptional<FVector2D>, DesiredSize)
 		SLATE_ATTRIBUTE(FFrameRate, DisplayRate)
 
 		SLATE_ARGUMENT(bool, ShowInputGridNumbers)
@@ -270,7 +269,7 @@ protected:
 	TAttribute<TOptional<float>> DataMinInput;
 	TAttribute<TOptional<float>> DataMaxInput;
 
-	TAttribute<FVector2D> DesiredSize;
+	TAttribute<TOptional<FVector2D>> DesiredSize;
 	TAttribute<FFrameRate> DisplayRate;
 
 	bool bShowInputGridNumbers = false;
diff --git a/Engine/Plugins/Experimental/EaseCurveTool/Source/EaseCurveTool/Private/Widgets/SEaseCurvePreset.cpp b/Engine/Plugins/Experimental/EaseCurveTool/Source/EaseCurveTool/Private/Widgets/SEaseCurvePreset.cpp
--- a/Engine/Plugins/Experimental/EaseCurveTool/Source/EaseCurveTool/Private/Widgets/SEaseCurvePreset.cpp
+++ b/Engine/Plugins/Experimental/EaseCurveTool/Source/EaseCurveTool/Private/Widgets/SEaseCurvePreset.cpp
@@ -133,13 +133,15 @@ void SEaseCurvePreset::HandleLibraryChanged(const TWeakObjectPtr<UEaseCurveLibra
 
 void SEaseCurvePreset::HandlePresetChanged()
 {
+	if (!PresetComboBox.IsValid())
+	{
+		return;
+	}
+
 	FEaseCurvePreset Preset;
 	const bool bHasSelectedItem = PresetComboBox->GetSelectedItem(Preset);
 
-	if (PresetComboBox.IsValid())
-	{
-		PresetComboBox->Reload();
-	}
+	PresetComboBox->Reload();
 
 	if (bHasSelectedItem)
 	{
diff --git a/Engine/Plugins/Experimental/EaseCurveTool/Source/EaseCurveTool/Private/Widgets/SEaseCurveTool.cpp b/Engine/Plugins/Experimental/EaseCurveTool/Source/EaseCurveTool/Private/Widgets/SEaseCurveTool.cpp
--- a/Engine/Plugins/Experimental/EaseCurveTool/Source/EaseCurveTool/Private/Widgets/SEaseCurveTool.cpp
+++ b/Engine/Plugins/Experimental/EaseCurveTool/Source/EaseCurveTool/Private/Widgets/SEaseCurveTool.cpp
@@ -37,12 +37,14 @@ void SEaseCurveTool::Construct(const FArguments& InArgs, const TSharedRef<FEaseC
 			+ SVerticalBox::Slot()
 			.AutoHeight()
 			.HAlign(HAlign_Fill)
-			.Padding(0.f, 1.f, 0.f, 0.f)
+			.Padding(0.f, 2.f, 0.f, 0.f)
 			[
 				SNew(SEaseCurveLibraryComboBox, InTool)
 			]
 			+ SVerticalBox::Slot()
 			.AutoHeight()
+			.Padding(0.f, 0.f, 0.f, 2.f)
+			.Expose(ToolWidgetSlot)
 			[
 				SNew(SVerticalBox)
 				.Visibility_Lambda([this]()
@@ -66,7 +68,8 @@ void SEaseCurveTool::Construct(const FArguments& InArgs, const TSharedRef<FEaseC
 				+ SVerticalBox::Slot()
 				.AutoHeight()
 				.HAlign(HAlign_Fill)
-				.Padding(0.f, 4.f, 0.f, 0.f)
+				.Padding(0.f, 3.f, 0.f, 0.f)
+				.Expose(GraphEditorSlot)
 				[
 					ConstructCurveEditorPanel()
 				]
@@ -95,6 +98,25 @@ void SEaseCurveTool::Construct(const FArguments& InArgs, const TSharedRef<FEaseC
 	CurvePresetWidget->SetSelectedItem(InArgs._InitialTangents);
 }
 
+FVector2D SEaseCurveTool::ComputeDesiredSize(const float InScaleMultiplier) const
+{
+	if (ToolWidgetSlot && GraphEditorSlot)
+	{
+		if (GetDefault<UEaseCurveToolSettings>()->IsToolVisible())
+		{
+			ToolWidgetSlot->SetFillHeight(1.f);
+			GraphEditorSlot->SetFillHeight(1.f);
+		}
+		else
+		{
+			ToolWidgetSlot->SetAutoHeight();
+			GraphEditorSlot->SetAutoHeight();
+		}
+	}
+
+	return SCompoundWidget::ComputeDesiredSize(InScaleMultiplier);
+}
+
 TSharedRef<SWidget> SEaseCurveTool::ConstructCurveEditorPanel()
 {
 	const TSharedPtr<FEaseCurveTool> EaseCurveTool = WeakTool.Pin();
@@ -109,6 +131,13 @@ TSharedRef<SWidget> SEaseCurveTool::ConstructCurveEditorPanel()
 			}));
 
 	return SNew(SBorder)
+		.BorderImage_Lambda([this]
+			{
+				const bool bIsFocusedOrHovered = CurveEaseEditorWidget.IsValid()
+					&& (CurveEaseEditorWidget->HasKeyboardFocus() || CurveEaseEditorWidget->IsHovered());
+				const FName BrushName = bIsFocusedOrHovered ? TEXT("Editor.Border.Hover") : TEXT("Editor.Border");
+				return FEaseCurveStyle::Get().GetBrush(BrushName);
+			})
 		[
 			SNew(SOverlay)
 			+ SOverlay::Slot()
@@ -180,9 +209,10 @@ EEaseCurveToolOperation SEaseCurveTool::GetToolOperation() const
 	return WeakTool.IsValid() ? WeakTool.Pin()->GetToolOperation() : EEaseCurveToolOperation::InOut;
 }
 
-FVector2D SEaseCurveTool::GetEditorSize() const
+TOptional<FVector2D> SEaseCurveTool::GetEditorSize() const
 {
-	return FVector2D(CurrentGraphSize);
+	return GetDefault<UEaseCurveToolSettings>()->IsToolVisible()
+		? TOptional<FVector2D>() : FVector2D(CurrentGraphSize);
 }
 
 bool SEaseCurveTool::CanEditCurve() const
diff --git a/Engine/Plugins/Experimental/EaseCurveTool/Source/EaseCurveTool/Private/Widgets/SEaseCurveTool.h b/Engine/Plugins/Experimental/EaseCurveTool/Source/EaseCurveTool/Private/Widgets/SEaseCurveTool.h
--- a/Engine/Plugins/Experimental/EaseCurveTool/Source/EaseCurveTool/Private/Widgets/SEaseCurveTool.h
+++ b/Engine/Plugins/Experimental/EaseCurveTool/Source/EaseCurveTool/Private/Widgets/SEaseCurveTool.h
@@ -49,11 +49,15 @@ public:
 	void ResetTangentsAndNotify() const;
 
 protected:
+	//~ Begin SWidget
+	virtual FVector2D ComputeDesiredSize(const float InScaleMultiplier) const override;
+	//~ End SWidget
+
 	TSharedRef<SWidget> ConstructCurveEditorPanel();
 
 	FFrameRate GetDisplayRate() const;
 	EEaseCurveToolOperation GetToolOperation() const;
-	FVector2D GetEditorSize() const;
+	TOptional<FVector2D> GetEditorSize() const;
 
 	bool CanEditCurve() const;
 	TOptional<FText> GetErrorMessage() const;
@@ -100,6 +104,9 @@ protected:
 	int32 CurrentGraphSize = DefaultGraphSize;
 
 	TSharedPtr<FEaseCurveToolContextMenu> ContextMenu;
+
+	SVerticalBox::FSlot* ToolWidgetSlot = nullptr;
+	SVerticalBox::FSlot* GraphEditorSlot = nullptr;
 };
 
 } // namespace UE::EaseCurveTool
