[EaseCurveTool]
- Disable curve editor toolbar button for sequencer selections
- Better tooltip for curve editor toolbar button

diff --git a/Engine/Plugins/Experimental/EaseCurveTool/Source/EaseCurveTool/Private/EaseCurveEditorExtension.cpp b/Engine/Plugins/Experimental/EaseCurveTool/Source/EaseCurveTool/Private/EaseCurveEditorExtension.cpp
--- a/Engine/Plugins/Experimental/EaseCurveTool/Source/EaseCurveTool/Private/EaseCurveEditorExtension.cpp
+++ b/Engine/Plugins/Experimental/EaseCurveTool/Source/EaseCurveTool/Private/EaseCurveEditorExtension.cpp
@@ -146,22 +146,37 @@ bool FEaseCurveEditorExtension::IsButtonEnabled() const
 	{
 		// In order to apply operations on keys other than cubic weighted, we skip checking
 		// ToolInstance->GetSelectionError() == EEaseCurveToolError::None
-		return ToolInstance->HasCachedKeysToEase();
+		return ToolInstance->IsCurveEditorSelection()
+			&& ToolInstance->HasCachedKeysToEase();
 	}
 	return false;
 }
 
 FText FEaseCurveEditorExtension::GetEaseCurveTooltipText() const
 {
-	if (const TSharedPtr<FEaseCurveTool> ToolInstance = GetToolInstance())
+	static const FText DefaultTooltipText = LOCTEXT("EaseCurveTooltip", "Apply an ease curve between the selected key tangents");
+
+	const TSharedPtr<FEaseCurveTool> ToolInstance = GetToolInstance();
+	if (!ToolInstance.IsValid())
 	{
-		if (ToolInstance->HasCachedKeysToEase()
-			&& ToolInstance->GetSelectionError() != EEaseCurveToolError::None)
-		{
-			return ToolInstance->GetSelectionErrorText();
-		}
+		return FText::Format(LOCTEXT("NoToolInstanceTooltip", "{0}\n\nError: No tool instance!")
+			, DefaultTooltipText);
 	}
-	return LOCTEXT("EaseCurvePresetsTooltip", "Apply an ease curve between the selected key tangents");
+
+	if (!ToolInstance->IsCurveEditorSelection()
+		|| !ToolInstance->HasCachedKeysToEase())
+	{
+		return FText::Format(LOCTEXT("NoSelectionTooltip", "{0}\n\nError: No selection!")
+			, DefaultTooltipText);
+	}
+
+	if (ToolInstance->GetSelectionError() != EEaseCurveToolError::None)
+	{
+		return FText::Format(LOCTEXT("NoEaseableKeysTooltip", "{0}\n\nError: {1}")
+			, DefaultTooltipText, ToolInstance->GetSelectionErrorText());
+	}
+
+	return DefaultTooltipText;
 }
 
 } // namespace UE::EaseCurveTool
diff --git a/Engine/Plugins/Experimental/EaseCurveTool/Source/EaseCurveTool/Private/EaseCurveTool.cpp b/Engine/Plugins/Experimental/EaseCurveTool/Source/EaseCurveTool/Private/EaseCurveTool.cpp
--- a/Engine/Plugins/Experimental/EaseCurveTool/Source/EaseCurveTool/Private/EaseCurveTool.cpp
+++ b/Engine/Plugins/Experimental/EaseCurveTool/Source/EaseCurveTool/Private/EaseCurveTool.cpp
@@ -695,7 +695,7 @@ FText FEaseCurveTool::GetSelectionErrorText() const
 	case EEaseCurveToolError::None:
 		return FText::GetEmpty();
 	case EEaseCurveToolError::LastKey:
-		return LOCTEXT("LastSelectedKey", "No next key to ease to");
+		return LOCTEXT("LastSelectedKey", "No next key to ease to!");
 	case EEaseCurveToolError::SameValues:
 		return LOCTEXT("EqualValueKeys", "No different key values!");
 	case EEaseCurveToolError::NoWeightedBrokenCubicTangents:
