[EaseCurveTool] Skip recording analytics for invalid presets

diff --git a/Engine/Plugins/Experimental/EaseCurveTool/Source/EaseCurveTool/Private/EaseCurveTool.cpp b/Engine/Plugins/Experimental/EaseCurveTool/Source/EaseCurveTool/Private/EaseCurveTool.cpp
--- a/Engine/Plugins/Experimental/EaseCurveTool/Source/EaseCurveTool/Private/EaseCurveTool.cpp
+++ b/Engine/Plugins/Experimental/EaseCurveTool/Source/EaseCurveTool/Private/EaseCurveTool.cpp
@@ -664,6 +664,11 @@ void FEaseCurveTool::ShowNotificationMessage(const FText& InMessageText)
 
 void FEaseCurveTool::RecordPresetAnalytics(const TSharedPtr<FEaseCurvePreset>& InPreset, const FString& InLocation)
 {
+	if (!FEngineAnalytics::IsAvailable() || !InPreset.IsValid())
+	{
+		return;
+	}
+
 	// Only send analytics for default presets
 	const TArray<FEaseCurvePreset> DefaultPresets = UEaseCurveLibrary::GetDefaultPresets();
 
diff --git a/Engine/Plugins/Experimental/EaseCurveTool/Source/EaseCurveTool/Private/Menus/CurveEditorToolBarMenu.cpp b/Engine/Plugins/Experimental/EaseCurveTool/Source/EaseCurveTool/Private/Menus/CurveEditorToolBarMenu.cpp
--- a/Engine/Plugins/Experimental/EaseCurveTool/Source/EaseCurveTool/Private/Menus/CurveEditorToolBarMenu.cpp
+++ b/Engine/Plugins/Experimental/EaseCurveTool/Source/EaseCurveTool/Private/Menus/CurveEditorToolBarMenu.cpp
@@ -136,10 +136,7 @@ void FCurveEditorToolBarMenu::HandlePresetChanged(const TSharedPtr<FEaseCurvePre
 		, /*bInBroadcastUpdate=*/true
 		, /*bInSetSequencerTangents=*/true);
 
-	if (FEngineAnalytics::IsAvailable())
-	{
-		ToolInstance->RecordPresetAnalytics(InPreset, TEXT("CurveEditorExtension"));
-	}
+	ToolInstance->RecordPresetAnalytics(InPreset, TEXT("CurveEditorExtension"));
 }
 
 void FCurveEditorToolBarMenu::HandleQuickPresetChanged(const TSharedPtr<FEaseCurvePreset>& InPreset)
diff --git a/Engine/Plugins/Experimental/EaseCurveTool/Source/EaseCurveTool/Private/Widgets/SEaseCurveTool.cpp b/Engine/Plugins/Experimental/EaseCurveTool/Source/EaseCurveTool/Private/Widgets/SEaseCurveTool.cpp
--- a/Engine/Plugins/Experimental/EaseCurveTool/Source/EaseCurveTool/Private/Widgets/SEaseCurveTool.cpp
+++ b/Engine/Plugins/Experimental/EaseCurveTool/Source/EaseCurveTool/Private/Widgets/SEaseCurveTool.cpp
@@ -400,14 +400,11 @@ void SEaseCurveTool::OnPresetChanged(const TSharedPtr<FEaseCurvePreset>& InPrese
 	if (InPreset.IsValid())
 	{
 		SetTangents(InPreset->Tangents, ToolOperation.Get(), true, true, true);
-	}
-
-	FSlateApplication::Get().SetAllUserFocus(CurveEaseEditorWidget);
 
-	if (FEngineAnalytics::IsAvailable())
-	{
 		EaseCurveTool->RecordPresetAnalytics(InPreset, TEXT("WithGraphEditor"));
 	}
+
+	FSlateApplication::Get().SetAllUserFocus(CurveEaseEditorWidget);
 }
 
 void SEaseCurveTool::OnQuickPresetChanged(const TSharedPtr<FEaseCurvePreset>& InPreset) const
