Curve Editor x Ease Curve Tool: Fix "Set Ease Curve Tangents" blocking undo when you move tangents.

The problem was that the FEaseCurveTool would detect a change made to selection at the end of tick an create a new transaction. Undoing would change selection. So if you had a big undo stack, undid a few steps, the actions past the action you undid would be replaced by the transaction incorrectly created by FEaseCurveTool at end of tick.

The solution is to let FEaseCurveTool make any changes as part of the curve editor transaction that changes the selection (during GUndo) and when playing back the undo action (during GIsTransacting).

diff --git a/Engine/Plugins/Experimental/EaseCurveTool/Source/EaseCurveTool/Private/EaseCurveTool.cpp b/Engine/Plugins/Experimental/EaseCurveTool/Source/EaseCurveTool/Private/EaseCurveTool.cpp
--- a/Engine/Plugins/Experimental/EaseCurveTool/Source/EaseCurveTool/Private/EaseCurveTool.cpp
+++ b/Engine/Plugins/Experimental/EaseCurveTool/Source/EaseCurveTool/Private/EaseCurveTool.cpp
@@ -49,6 +49,10 @@ FEaseCurveTool::FEaseCurveTool(const TSharedRef<ISequencer>& InSequencer)
 
 	InSequencer->GetSelectionChangedObjectGuids().AddRaw(this, &FEaseCurveTool::OnSequencerSelectionChanged);
 	InSequencer->OnMovieSceneDataChanged().AddRaw(this, &FEaseCurveTool::OnMovieSceneDataChanged);
+	if (const TSharedPtr<FCurveEditor> CurveEditor = FEaseCurveToolUtils::GetCurveEditorFromSequencer(InSequencer))
+	{
+		CurveEditor->Selection.OnSelectionChanged().AddRaw(this, &FEaseCurveTool::OnCurveEditorSelectionChanged);
+	}
 
 	ToolTab->OnVisibilityChanged().AddRaw(this, &FEaseCurveTool::OnTabVisibilityChanged);
 
@@ -79,6 +83,11 @@ FEaseCurveTool::~FEaseCurveTool()
 	{
 		Sequencer->GetSelectionChangedObjectGuids().RemoveAll(this);
 		Sequencer->OnMovieSceneDataChanged().RemoveAll(this);
+		
+		if (const TSharedPtr<FCurveEditor> CurveEditor = FEaseCurveToolUtils::GetCurveEditorFromSequencer(Sequencer.ToSharedRef()))
+		{
+			CurveEditor->Selection.OnSelectionChanged().RemoveAll(this);
+		}
 	}
 
 	ToolTab->OnVisibilityChanged().RemoveAll(this);
@@ -129,6 +138,30 @@ void FEaseCurveTool::OnTabVisibilityChanged(const bool bInVisible)
 {
 }
 
+void FEaseCurveTool::OnCurveEditorSelectionChanged()
+{
+	// OnCurveEditorSelectionChanged is usually invoked as part of a transaction that changes selection, or as part of an undo / redo.
+	// Generally, changes to selection in FCurveEditor should be wrapped by FScopedSelectionChangeEventSuppression, which effectively bundles multiple
+	// changes into one broadcast at the end; FScopedSelectionChange constructs FScopedSelectionChangeEventSuppression. So, at this point, we needn't
+	// worry about OnCurveEditorSelectionChanged being called multiple times in the same frame.
+	if (
+		// As UpdateEaseCurveFromKeySelections may in turn start a FScopedTransaction, we must process the change immediately so our changes become
+		// part of the same transaction. 
+		GUndo
+		// If selection changes as part of an undo / redo, we must process immediately so any FScopedTransaction we start is not appended to the undo
+		// buffer. If we didn't, then at end of tick we would create a new transaction, which would override all the undo stack past the current
+		// transaction.
+		|| GIsTransacting)
+	{
+		UpdateEaseCurveFromKeySelections();
+	}
+	else
+	{
+		// Otherwise we'll optimize a little by deferring the update in case more changes are made.
+		bRefreshForKeySelectionChange = true;
+	}
+}
+
 TSharedRef<SWidget> FEaseCurveTool::GenerateWidget()
 {
 	UpdateEaseCurveFromKeySelections();
@@ -370,8 +403,9 @@ void FEaseCurveTool::SetSequencerKeySelectionTangents(const FEaseCurveTangents&
 
 void FEaseCurveTool::UpdateEaseCurveFromKeySelections()
 {
-	KeyCache = FEaseCurveKeySelection(SequencerWeak.Pin());
+	bRefreshForKeySelectionChange = false;
 
+	KeyCache = FEaseCurveKeySelection(SequencerWeak.Pin());
 	UpdateEaseCurveFromKeyCache();
 }
 
@@ -605,26 +639,9 @@ void FEaseCurveTool::Tick(float InDeltaTime)
 {
 	using namespace UE::EaseCurveTool;
 
-	// Since the curve editor does not have an event to subscribe to for selection change notifications,
-	// we are manually polling the selection each tick and comparing to the previous tick to be alerted
-	// of curve editor key selection changes.
-	if (const TSharedPtr<ISequencer> Sequencer = SequencerWeak.Pin())
-	{
-		if (const TSharedPtr<FCurveEditor> CurveEditor = FEaseCurveToolUtils::GetCurveEditorFromSequencer(Sequencer.ToSharedRef()))
-		{
-			const FCurveEditorSelection& CurveEditorSelection = CurveEditor->GetSelection();
-			if (!FEaseCurveToolUtils::CompareCurveEditorSelections(CurveEditorSelection, LastCurveEditorSelection))
-			{
-				bRefreshForKeySelectionChange = true;
-			}
-			LastCurveEditorSelection = CurveEditorSelection;
-		}
-	}
-
 	if (bRefreshForKeySelectionChange)
 	{
 		UpdateEaseCurveFromKeySelections();
-		bRefreshForKeySelectionChange = false;
 	}
 }
 
diff --git a/Engine/Plugins/Experimental/EaseCurveTool/Source/EaseCurveTool/Private/EaseCurveTool.h b/Engine/Plugins/Experimental/EaseCurveTool/Source/EaseCurveTool/Private/EaseCurveTool.h
--- a/Engine/Plugins/Experimental/EaseCurveTool/Source/EaseCurveTool/Private/EaseCurveTool.h
+++ b/Engine/Plugins/Experimental/EaseCurveTool/Source/EaseCurveTool/Private/EaseCurveTool.h
@@ -211,6 +211,8 @@ private:
 	void OnMovieSceneDataChanged(const EMovieSceneDataChangeType InMetaData);
 
 	void OnTabVisibilityChanged(const bool bInVisible);
+	
+	void OnCurveEditorSelectionChanged();
 };
 
 } // namespace UE::EaseCurveTool
