[EaseCurveTool]
- Load default library on startup
- Disable clearing of library to force users to always have a library selected

diff --git a/Engine/Plugins/Experimental/EaseCurveTool/Source/EaseCurveTool/Private/EaseCurveToolSettings.h b/Engine/Plugins/Experimental/EaseCurveTool/Source/EaseCurveTool/Private/EaseCurveToolSettings.h
--- a/Engine/Plugins/Experimental/EaseCurveTool/Source/EaseCurveTool/Private/EaseCurveToolSettings.h
+++ b/Engine/Plugins/Experimental/EaseCurveTool/Source/EaseCurveTool/Private/EaseCurveToolSettings.h
@@ -28,7 +28,7 @@ public:
 	void SetShowCurveEditorToolbar(const bool bInShowCurveEditorToolbarButton) { bShowCurveEditorToolbarButton = bInShowCurveEditorToolbarButton; }
 	void ToggleShowCurveEditorToolbar() { bShowCurveEditorToolbarButton = !bShowCurveEditorToolbarButton; }
 
-	TObjectPtr<UEaseCurveLibrary> GetPresetLibrary() const { return DefaultPresetLibrary.Get(); }
+	TObjectPtr<UEaseCurveLibrary> GetPresetLibrary() const { return DefaultPresetLibrary.LoadSynchronous(); }
 	void SetPresetLibrary(const TObjectPtr<UEaseCurveLibrary>& InWeakPresetLibrary) { DefaultPresetLibrary = InWeakPresetLibrary; }
 
 	FText GetNewPresetCategory() const { return NewPresetCategory; }
diff --git a/Engine/Plugins/Experimental/EaseCurveTool/Source/EaseCurveTool/Private/Widgets/SEaseCurveLibraryComboBox.cpp b/Engine/Plugins/Experimental/EaseCurveTool/Source/EaseCurveTool/Private/Widgets/SEaseCurveLibraryComboBox.cpp
--- a/Engine/Plugins/Experimental/EaseCurveTool/Source/EaseCurveTool/Private/Widgets/SEaseCurveLibraryComboBox.cpp
+++ b/Engine/Plugins/Experimental/EaseCurveTool/Source/EaseCurveTool/Private/Widgets/SEaseCurveLibraryComboBox.cpp
@@ -47,7 +47,7 @@ void SEaseCurveLibraryComboBox::Construct(const FArguments& InArgs, const TShare
 			SNew(SObjectPropertyEntryBox)
 			.AllowedClass(UEaseCurveLibrary::StaticClass())
 			.NewAssetFactories(TArray<UFactory*> { PresetLibraryFactory.Get() })
-			.AllowClear(true)
+			.AllowClear(false)
 			.AllowCreate(true)
 			.DisplayThumbnail(true)
 			.DisplayUseSelected(true)
