[EaseCurveTool] Serializer API changes to avoid having to deprecate in 5.8

diff --git a/Engine/Plugins/Experimental/EaseCurveTool/Source/EaseCurveTool/Private/EaseCurveLibrary.cpp b/Engine/Plugins/Experimental/EaseCurveTool/Source/EaseCurveTool/Private/EaseCurveLibrary.cpp
--- a/Engine/Plugins/Experimental/EaseCurveTool/Source/EaseCurveTool/Private/EaseCurveLibrary.cpp
+++ b/Engine/Plugins/Experimental/EaseCurveTool/Source/EaseCurveTool/Private/EaseCurveLibrary.cpp
@@ -28,7 +28,7 @@ TArray<FEaseCurvePreset> UEaseCurveLibrary::GetDefaultPresets()
 	auto AddPreset = [&DefaultPresets](const FText& InName
 		, const FText& InCategory, const TStaticArray<double, 4>& InPoints)
 	{
-		DefaultPresets.Add(FEaseCurvePreset(InName, InCategory
+		DefaultPresets.Add(FEaseCurvePreset(InCategory, InName
 			, FEaseCurveTangents::MakeFromCubicBezier(InPoints)));
 	};
 
@@ -155,7 +155,7 @@ bool UEaseCurveLibrary::AddPresetToNewCategory(const FText& InName, const FEaseC
 		return false;
 	}
 
-	OutNewPreset = FEaseCurvePreset(InName, NewCategory, InTangents);
+	OutNewPreset = FEaseCurvePreset(NewCategory, InName, InTangents);
 
 	return AddPreset(OutNewPreset);
 }
diff --git a/Engine/Plugins/Experimental/EaseCurveTool/Source/EaseCurveTool/Private/EaseCurvePreset.h b/Engine/Plugins/Experimental/EaseCurveTool/Source/EaseCurveTool/Private/EaseCurvePreset.h
--- a/Engine/Plugins/Experimental/EaseCurveTool/Source/EaseCurveTool/Private/EaseCurvePreset.h
+++ b/Engine/Plugins/Experimental/EaseCurveTool/Source/EaseCurveTool/Private/EaseCurvePreset.h
@@ -33,11 +33,11 @@ struct FEaseCurvePreset
 	GENERATED_BODY()
 
 	FEaseCurvePreset() {}
-	FEaseCurvePreset(const FText& InName, const FText& InCategory, const FEaseCurveTangents& InTangents)
-		: Name(InName), Category(InCategory), Tangents(InTangents)
+	FEaseCurvePreset(const FText& InCategory, const FText& InName, const FEaseCurveTangents& InTangents)
+		: Category(InCategory), Name(InName), Tangents(InTangents)
 	{}
 	FEaseCurvePreset(const FEaseCurvePresetHandle& InHandle, const FEaseCurveTangents& InTangents)
-		: Name(InHandle.Name), Category(InHandle.Category), Tangents(InTangents)
+		: Category(InHandle.Category), Name(InHandle.Name), Tangents(InTangents)
 	{}
 
 	FORCEINLINE bool operator==(const FEaseCurvePresetHandle& InRhs) const
@@ -67,10 +67,10 @@ struct FEaseCurvePreset
 	}
 
 	UPROPERTY(Config, EditAnywhere, Category = "EaseCurvePreset")
-	FText Name;
+	FText Category;
 
 	UPROPERTY(Config, EditAnywhere, Category = "EaseCurvePreset")
-	FText Category;
+	FText Name;
 
 	UPROPERTY(Config, EditAnywhere, Category = "EaseCurvePreset")
 	FEaseCurveTangents Tangents;
diff --git a/Engine/Plugins/Experimental/EaseCurveTool/Source/EaseCurveTool/Private/EaseCurveSerializer.cpp b/Engine/Plugins/Experimental/EaseCurveTool/Source/EaseCurveTool/Private/EaseCurveSerializer.cpp
new file mode 100644
--- /dev/null
+++ b/Engine/Plugins/Experimental/EaseCurveTool/Source/EaseCurveTool/Private/EaseCurveSerializer.cpp
@@ -0,0 +1,99 @@
+﻿// Copyright Epic Games, Inc. All Rights Reserved.
+
+#include "EaseCurveSerializer.h"
+#include "DesktopPlatformModule.h"
+#include "Framework/Application/SlateApplication.h"
+#include "IDesktopPlatform.h"
+#include "Internationalization/Text.h"
+#include "Misc/Paths.h"
+#include "EaseCurveLibrary.h"
+
+#define LOCTEXT_NAMESPACE "EaseCurveSerializer"
+
+bool UEaseCurveSerializer::PromptUserForFilePath(FString& OutFilePath, const bool bInImport)
+{
+	IDesktopPlatform* const DesktopPlatform = FDesktopPlatformModule::Get();
+	if (!DesktopPlatform)
+	{
+		return false;
+	}
+
+	const void* const ParentWindowHandle = FSlateApplication::Get().FindBestParentWindowHandleForDialogs(nullptr);
+	const FString Title = TEXT("Select File Location");
+	const FString FileTypes = TEXT("All Files (*.*)|*.*");
+
+	bool bSuccess;
+	TArray<FString> OutFilePaths;
+
+	if (bInImport)
+	{
+		bSuccess = DesktopPlatform->OpenFileDialog(ParentWindowHandle
+			, Title
+			, FPaths::GetPath(OutFilePath)
+			, FPaths::GetBaseFilename(OutFilePath)
+			, FileTypes
+			, EFileDialogFlags::None
+			, OutFilePaths);
+	}
+	else
+	{
+		bSuccess = DesktopPlatform->SaveFileDialog(ParentWindowHandle
+			, Title
+			, FPaths::GetPath(OutFilePath)
+			, FPaths::GetBaseFilename(OutFilePath)
+			, FileTypes
+			, EFileDialogFlags::None
+			, OutFilePaths);
+	}
+
+	if (!bSuccess || OutFilePaths.Num() != 1)
+	{
+		return false;
+	}
+
+	OutFilePath = OutFilePaths[0];
+
+	return true;
+}
+
+FText UEaseCurveSerializer::GetDisplayName() const
+{
+	return FText::GetEmpty();
+}
+
+FText UEaseCurveSerializer::GetDisplayTooltip() const
+{
+	return FText::GetEmpty();
+}
+
+bool UEaseCurveSerializer::IsFileExport() const
+{
+	return true;
+}
+
+bool UEaseCurveSerializer::SupportsExport() const
+{
+	return false;
+}
+
+bool UEaseCurveSerializer::Export(const FString& InFilePath, TSet<TWeakObjectPtr<UEaseCurveLibrary>> InWeakLibraries)
+{
+	return false;
+}
+
+bool UEaseCurveSerializer::IsFileImport() const
+{
+	return true;
+}
+
+bool UEaseCurveSerializer::SupportsImport() const
+{
+	return false;
+}
+
+bool UEaseCurveSerializer::Import(const FString& InFilePath, TSet<TWeakObjectPtr<UEaseCurveLibrary>> InWeakLibraries)
+{
+	return false;
+}
+
+#undef LOCTEXT_NAMESPACE
diff --git a/Engine/Plugins/Experimental/EaseCurveTool/Source/EaseCurveTool/Private/Menus/EaseCurveLibraryMenu.cpp b/Engine/Plugins/Experimental/EaseCurveTool/Source/EaseCurveTool/Private/Menus/EaseCurveLibraryMenu.cpp
--- a/Engine/Plugins/Experimental/EaseCurveTool/Source/EaseCurveTool/Private/Menus/EaseCurveLibraryMenu.cpp
+++ b/Engine/Plugins/Experimental/EaseCurveTool/Source/EaseCurveTool/Private/Menus/EaseCurveLibraryMenu.cpp
@@ -63,13 +63,22 @@ bool FEaseCurveLibraryMenu::AddMenuEntryForSerializer(FMenuBuilder& InMenuBuilde
 	if (InSerializer.HasAnyFlags(RF_ClassDefaultObject)
 		&& !InSerializer.GetClass()->HasAnyClassFlags(CLASS_Deprecated | CLASS_Abstract))
 	{
-		const FName IconBrushName = bInImport ? TEXT("Icons.Toolbar.Import") : TEXT("Icons.Toolbar.Export");
-
-		InMenuBuilder.AddMenuEntry(InSerializer.GetDisplayName(),
-			InSerializer.GetDisplayTooltip(),
-			FSlateIcon(FAppStyle::GetAppStyleSetName(), IconBrushName),
-			FExecuteAction::CreateStatic(&FEaseCurveLibraryMenu::PromptForExport, &InSerializer, InWeakLibraryAssets)
-		);
+		if (bInImport)
+		{
+			InMenuBuilder.AddMenuEntry(InSerializer.GetDisplayName(),
+				InSerializer.GetDisplayTooltip(),
+				FSlateIcon(FAppStyle::GetAppStyleSetName(), TEXT("Icons.Toolbar.Import")),
+				FExecuteAction::CreateStatic(&FEaseCurveLibraryMenu::PromptForImport, &InSerializer, InWeakLibraryAssets)
+			);
+		}
+		else
+		{
+			InMenuBuilder.AddMenuEntry(InSerializer.GetDisplayName(),
+				InSerializer.GetDisplayTooltip(),
+				FSlateIcon(FAppStyle::GetAppStyleSetName(), TEXT("Icons.Toolbar.Export")),
+				FExecuteAction::CreateStatic(&FEaseCurveLibraryMenu::PromptForExport, &InSerializer, InWeakLibraryAssets)
+			);
+		}
 
 		return true;
 	}
@@ -119,7 +128,8 @@ void FEaseCurveLibraryMenu::PopulateImportMenu(FMenuBuilder& InMenuBuilder
 	{
 		if (UEaseCurveSerializer* const Serializer = *SerializerIt)
 		{
-			if (AddMenuEntryForSerializer(InMenuBuilder, *Serializer, /*bInImport=*/true, InWeakLibraryAssets))
+			if (Serializer->SupportsImport()
+				&& AddMenuEntryForSerializer(InMenuBuilder, *Serializer, /*bInImport=*/true, InWeakLibraryAssets))
 			{
 				bIsEmpty = false;
 			}
@@ -146,7 +156,8 @@ void FEaseCurveLibraryMenu::PopulateExportMenu(FMenuBuilder& InMenuBuilder
 	{
 		if (UEaseCurveSerializer* const Serializer = *SerializerIt)
 		{
-			if (AddMenuEntryForSerializer(InMenuBuilder, *Serializer, /*bInImport=*/false, InWeakLibraryAssets))
+			if (Serializer->SupportsExport()
+				&& AddMenuEntryForSerializer(InMenuBuilder, *Serializer, /*bInImport=*/false, InWeakLibraryAssets))
 			{
 				bIsEmpty = false;
 			}
@@ -165,45 +176,41 @@ void FEaseCurveLibraryMenu::PopulateExportMenu(FMenuBuilder& InMenuBuilder
 void FEaseCurveLibraryMenu::PromptForImport(UEaseCurveSerializer* const InSerializer
 	, const TSet<TWeakObjectPtr<UEaseCurveLibrary>> InWeakLibraryAssets)
 {
-	FString FilePath;
-	if (!InSerializer || !InSerializer->PromptUserForFilePath(FilePath, /*bInImport=*/true))
+	if (!InSerializer || !InSerializer->SupportsImport())
 	{
 		return;
 	}
 
-	TSet<UEaseCurveLibrary*> LibraryAssets;
+	FString FilePath;
 
-	for (const TWeakObjectPtr<UEaseCurveLibrary>& WeakLibraryAsset : InWeakLibraryAssets)
+	// Check for cancel if using file prompt
+	if (InSerializer->IsFileImport()
+		&& !InSerializer->PromptUserForFilePath(FilePath, /*bInImport=*/true))
 	{
-		if (WeakLibraryAsset.IsValid())
-		{
-			LibraryAssets.Add(WeakLibraryAsset.Get());
-		}
+		return;
 	}
 
-	InSerializer->Import(FilePath, LibraryAssets);
+	InSerializer->Import(FilePath, InWeakLibraryAssets);
 }
 
 void FEaseCurveLibraryMenu::PromptForExport(UEaseCurveSerializer* const InSerializer
 	, const TSet<TWeakObjectPtr<UEaseCurveLibrary>> InWeakLibraryAssets)
 {
-	FString FilePath;
-	if (!InSerializer || !InSerializer->PromptUserForFilePath(FilePath, /*bInImport=*/false))
+	if (!InSerializer || !InSerializer->SupportsExport())
 	{
 		return;
 	}
 
-	TSet<UEaseCurveLibrary*> LibraryAssets;
+	FString FilePath;
 
-	for (const TWeakObjectPtr<UEaseCurveLibrary>& WeakLibraryAsset : InWeakLibraryAssets)
+	// Check for cancel if using file prompt
+	if (InSerializer->IsFileExport()
+		&& !InSerializer->PromptUserForFilePath(FilePath, /*bInImport=*/false))
 	{
-		if (WeakLibraryAsset.IsValid())
-		{
-			LibraryAssets.Add(WeakLibraryAsset.Get());
-		}
+		return;
 	}
 
-	InSerializer->Export(FilePath, LibraryAssets);
+	InSerializer->Export(FilePath, InWeakLibraryAssets);
 }
 
 void FEaseCurveLibraryMenu::SetToDefaultPresets(const TSet<TWeakObjectPtr<UEaseCurveLibrary>> InWeakLibraryAssets)
diff --git a/Engine/Plugins/Experimental/EaseCurveTool/Source/EaseCurveTool/Public/EaseCurveSerializer.h b/Engine/Plugins/Experimental/EaseCurveTool/Source/EaseCurveTool/Public/EaseCurveSerializer.h
--- a/Engine/Plugins/Experimental/EaseCurveTool/Source/EaseCurveTool/Public/EaseCurveSerializer.h
+++ b/Engine/Plugins/Experimental/EaseCurveTool/Source/EaseCurveTool/Public/EaseCurveSerializer.h
@@ -2,16 +2,13 @@
 
 #pragma once
 
-#include "DesktopPlatformModule.h"
-#include "Framework/Application/SlateApplication.h"
-#include "IDesktopPlatform.h"
-#include "Internationalization/Text.h"
-#include "Misc/Paths.h"
+#include "Containers/ContainersFwd.h"
+#include "Containers/StringFwd.h"
+#include "UObject/Object.h"
 #include "EaseCurveSerializer.generated.h"
 
-class UEaseCurveLibrary;
-class FString;
 class FText;
+class UEaseCurveLibrary;
 
 #define UE_API EASECURVETOOL_API
 
@@ -25,91 +22,47 @@ class UEaseCurveSerializer : public UObject
 	GENERATED_BODY()
 
 public:
-	static bool PromptUserForFilePath(FString& OutFilePath, const bool bInImport)
-	{
-		IDesktopPlatform* const DesktopPlatform = FDesktopPlatformModule::Get();
-		if (!DesktopPlatform)
-		{
-			return false;
-		}
-
-		const void* const ParentWindowHandle = FSlateApplication::Get().FindBestParentWindowHandleForDialogs(nullptr);
-		const FString Title = TEXT("Select File Location");
-		const FString FileTypes = TEXT("All Files (*.*)|*.*");
-
-		bool bSuccess;
-		TArray<FString> OutFilePaths;
-
-		if (bInImport)
-		{
-			bSuccess = DesktopPlatform->OpenFileDialog(ParentWindowHandle
-				, Title
-				, FPaths::GetPath(OutFilePath)
-				, FPaths::GetBaseFilename(OutFilePath)
-				, FileTypes
-				, EFileDialogFlags::None
-				, OutFilePaths);
-		}
-		else
-		{
-			bSuccess = DesktopPlatform->SaveFileDialog(ParentWindowHandle
-				, Title
-				, FPaths::GetPath(OutFilePath)
-				, FPaths::GetBaseFilename(OutFilePath)
-				, FileTypes
-				, EFileDialogFlags::None
-				, OutFilePaths);
-		}
-
-		if (!bSuccess || OutFilePaths.Num() != 1)
-		{
-			return false;
-		}
-
-		OutFilePath = OutFilePaths[0];
-
-		return true;
-	}
+	static bool PromptUserForFilePath(FString& OutFilePath, const bool bInImport);
 
 	/** The display name text to show for menu entries */
-	UE_API virtual FText GetDisplayName() const
-	{
-		return FText::GetEmpty();
-	}
+	UE_API virtual FText GetDisplayName() const;
 
 	/** The display tooltip text to show for menu entries */
-	UE_API virtual FText GetDisplayTooltip() const
-	{
-		return FText::GetEmpty();
-	}
+	UE_API virtual FText GetDisplayTooltip() const;
+
+	/** @return True if the export operation should prompt for a file location */
+	UE_API virtual bool IsFileExport() const;
+
+	/** @return True if this serializer supports export of ease curve presets */
+	UE_API virtual bool SupportsExport() const;
 
 	/**
 	 * Exports an Ease Curve Library asset to an external file.
 	 * Override to implement custom export logic (for a specific file type for example)
 	 *
 	 * @param InFilePath The path to the file where the library will be exported
-	 * @param InLibraries The Ease Curve Library assets to export
+	 * @param InWeakLibraries The Ease Curve Library assets to export
 	 *
 	 * @return True if the export operation was successful, false otherwise
 	 */
-	UE_API virtual bool Export(const FString& InFilePath, TSet<UEaseCurveLibrary*> InLibraries)
-	{
-		return false;
-	}
+	UE_API virtual bool Export(const FString& InFilePath, TSet<TWeakObjectPtr<UEaseCurveLibrary>> InWeakLibraries);
+
+	/** @return True if the import operation should prompt for a file location */
+	UE_API virtual bool IsFileImport() const;
+
+	/** @return True if this serializer supports import of ease curve presets */
+	UE_API virtual bool SupportsImport() const;
 
 	/**
 	 * Imports an Ease Curve Library asset from an external file.
 	 * Override to implement custom import logic (for a specific file type for example)
 	 *
 	 * @param InFilePath The path to the file where the library will be exported
-	 * @param InLibraries The Ease Curve Library assets to import
+	 * @param InWeakLibraries The Ease Curve Library assets to import
 	 *
 	 * @return True if the import operation was successful, false otherwise
 	 */
-    UE_API virtual bool Import(const FString& InFilePath, TSet<UEaseCurveLibrary*> InLibraries)
-	{
-		return false;
-	}
+    UE_API virtual bool Import(const FString& InFilePath, TSet<TWeakObjectPtr<UEaseCurveLibrary>> InWeakLibraries);
 };
 
 #undef UE_API
