[EaseCurveTool] Move plugin out of "Experimental" folder and into "Animation"

#virtualized

diff --git a/Engine/Plugins/Animation/EaseCurveTool/Source/EaseCurveTool/Private/EaseCurveToolModule.cpp b/Engine/Plugins/Animation/EaseCurveTool/Source/EaseCurveTool/Private/EaseCurveToolModule.cpp
new file mode 100644
--- /dev/null
+++ b/Engine/Plugins/Animation/EaseCurveTool/Source/EaseCurveTool/Private/EaseCurveToolModule.cpp
@@ -0,0 +1,100 @@
+﻿// Copyright Epic Games, Inc. All Rights Reserved.
+
+#include "EaseCurveToolModule.h"
+#include "ContentBrowserModule.h"
+#include "EaseCurveEditorExtension.h"
+#include "EaseCurveStyle.h"
+#include "EaseCurveToolCommands.h"
+#include "EaseCurveToolExtender.h"
+#include "Framework/MultiBox/MultiBoxExtender.h"
+#include "ICurveEditorModule.h"
+#include "Menus/EaseCurveLibraryMenu.h"
+
+void FEaseCurveToolModule::StartupModule()
+{
+	using namespace UE::EaseCurveTool;
+
+	FEaseCurveToolCommands::Register();
+
+	// Ensure singleton instances are created
+	FEaseCurveStyle::Get();
+	FEaseCurveToolExtender::Get();
+
+	RegisterContentBrowserExtender();
+	RegisterCurveEditorExtender();
+}
+
+void FEaseCurveToolModule::ShutdownModule()
+{
+	UnregisterCurveEditorExtender();
+	UnregisterContentBrowserExtender();
+
+	FEaseCurveToolCommands::Unregister();
+}
+
+void FEaseCurveToolModule::RegisterContentBrowserExtender()
+{
+	FContentBrowserModule* const ContentBrowserModule = FModuleManager::GetModulePtr<FContentBrowserModule>(TEXT("ContentBrowser"));
+	if (!ContentBrowserModule)
+	{
+		return;
+	}
+
+	TArray<FContentBrowserMenuExtender_SelectedAssets>& AssetMenuExtenders = ContentBrowserModule->GetAllAssetViewContextMenuExtenders();
+
+	AssetMenuExtenders.Add(FContentBrowserMenuExtender_SelectedAssets::CreateStatic(&UE::EaseCurveTool::FEaseCurveLibraryMenu::OnExtendContentBrowserAssetSelectionMenu));
+	ContentBrowserExtenderHandle = AssetMenuExtenders.Last().GetHandle();
+}
+
+void FEaseCurveToolModule::UnregisterContentBrowserExtender()
+{
+	if (!ContentBrowserExtenderHandle.IsValid())
+	{
+		return;
+	}
+
+	if (FContentBrowserModule* const ContentBrowserModule = FModuleManager::GetModulePtr<FContentBrowserModule>(TEXT("ContentBrowser")))
+	{
+		TArray<FContentBrowserMenuExtender_SelectedAssets>& AssetMenuExtenders = ContentBrowserModule->GetAllAssetViewContextMenuExtenders();
+		AssetMenuExtenders.RemoveAll([this](const FContentBrowserMenuExtender_SelectedAssets& InDelegate)
+			{
+				return InDelegate.GetHandle() == ContentBrowserExtenderHandle;
+			});
+	}
+
+	ContentBrowserExtenderHandle.Reset();
+}
+
+void FEaseCurveToolModule::RegisterCurveEditorExtender()
+{
+	using namespace UE::EaseCurveTool;
+
+	ICurveEditorModule* const CurveEditorModule = FModuleManager::GetModulePtr<ICurveEditorModule>(TEXT("CurveEditor"));
+	if (!CurveEditorModule)
+	{
+		return;
+	}
+
+	CurveEditorExtenderHandle = CurveEditorModule->RegisterEditorExtension(
+		FOnCreateCurveEditorExtension::CreateLambda([](const TWeakPtr<FCurveEditor> InCurveEditor)
+		{
+			return MakeShared<FEaseCurveEditorExtension>(InCurveEditor);
+		}));
+}
+
+void FEaseCurveToolModule::UnregisterCurveEditorExtender()
+{
+	if (!CurveEditorExtenderHandle.IsValid())
+	{
+		return;
+	}
+
+	if (ICurveEditorModule* const CurveEditorModule = FModuleManager::GetModulePtr<ICurveEditorModule>(TEXT("CurveEditor")))
+	{
+		CurveEditorModule->UnregisterEditorExtension(CurveEditorExtenderHandle);
+	}
+
+	CurveEditorExtenderHandle.Reset();
+}
+
+IMPLEMENT_MODULE(FEaseCurveToolModule, EaseCurveTool)
