Fix the handling of PostUndo / PostRedo. The original intent was to call UpdateEaseCurveFromKeySelections().

#rnx

diff --git a/Engine/Plugins/Animation/EaseCurveTool/Source/EaseCurveTool/Private/EaseCurveTool.cpp b/Engine/Plugins/Animation/EaseCurveTool/Source/EaseCurveTool/Private/EaseCurveTool.cpp
--- a/Engine/Plugins/Animation/EaseCurveTool/Source/EaseCurveTool/Private/EaseCurveTool.cpp
+++ b/Engine/Plugins/Animation/EaseCurveTool/Source/EaseCurveTool/Private/EaseCurveTool.cpp
@@ -612,12 +612,12 @@ void FEaseCurveTool::RedoAction()
 
 void FEaseCurveTool::PostUndo(bool bInSuccess)
 {
-	UndoAction();
+	UpdateEaseCurveFromKeySelections();
 }
 
 void FEaseCurveTool::PostRedo(bool bInSuccess)
 {
-	RedoAction();
+	UpdateEaseCurveFromKeySelections();
 }
 
 bool FEaseCurveTool::IsTickable() const
