[EaseCurveTool] Add serializer to import cubic bezier curves

diff --git a/Engine/Plugins/Experimental/EaseCurveTool/Source/EaseCurveTool/Private/Serializers/CubicBezierCurveSerializer.cpp b/Engine/Plugins/Experimental/EaseCurveTool/Source/EaseCurveTool/Private/Serializers/CubicBezierCurveSerializer.cpp
new file mode 100644
--- /dev/null
+++ b/Engine/Plugins/Experimental/EaseCurveTool/Source/EaseCurveTool/Private/Serializers/CubicBezierCurveSerializer.cpp
@@ -0,0 +1,51 @@
+﻿// Copyright Epic Games, Inc. All Rights Reserved.
+
+#include "Serializers/CubicBezierCurveSerializer.h"
+#include "CubicBezierPointEntryDialog.h"
+#include "EaseCurveLibrary.h"
+#include "Internationalization/Text.h"
+
+#define LOCTEXT_NAMESPACE "CubicBezierCurveSerializer"
+
+FText UCubicBezierCurveSerializer::GetDisplayName() const
+{
+	return LOCTEXT("CubicBezierCurve_Name", "Cubic Bezier Curve");
+}
+
+/** The display tooltip text to show for menu entries */
+FText UCubicBezierCurveSerializer::GetDisplayTooltip() const
+{
+	return LOCTEXT("CubicBezierCurve_Tooltip", "Import a cubic bezier curve as a preset");
+}
+
+bool UCubicBezierCurveSerializer::IsFileExport() const
+{
+	return false;
+}
+
+bool UCubicBezierCurveSerializer::SupportsExport() const
+{
+	return false;
+}
+
+bool UCubicBezierCurveSerializer::Export(const FString& InFilePath, TSet<TWeakObjectPtr<UEaseCurveLibrary>> InWeakLibraries)
+{
+	return false;
+}
+
+bool UCubicBezierCurveSerializer::IsFileImport() const
+{
+	return false;
+}
+
+bool UCubicBezierCurveSerializer::SupportsImport() const
+{
+	return true;
+}
+
+bool UCubicBezierCurveSerializer::Import(const FString& InFilePath, TSet<TWeakObjectPtr<UEaseCurveLibrary>> InWeakLibraries)
+{
+	return UE::EaseCurveTool::FCubicBezierPointEntryDialog::Prompt(InWeakLibraries);
+}
+
+#undef LOCTEXT_NAMESPACE
diff --git a/Engine/Plugins/Experimental/EaseCurveTool/Source/EaseCurveTool/Private/Serializers/CubicBezierCurveSerializer.h b/Engine/Plugins/Experimental/EaseCurveTool/Source/EaseCurveTool/Private/Serializers/CubicBezierCurveSerializer.h
new file mode 100644
--- /dev/null
+++ b/Engine/Plugins/Experimental/EaseCurveTool/Source/EaseCurveTool/Private/Serializers/CubicBezierCurveSerializer.h
@@ -0,0 +1,38 @@
+﻿// Copyright Epic Games, Inc. All Rights Reserved.
+
+#pragma once
+
+#include "Containers/ContainersFwd.h"
+#include "Containers/StringFwd.h"
+#include "EaseCurveSerializer.h"
+#include "CubicBezierCurveSerializer.generated.h"
+
+class UEaseCurveLibrary;
+class FText;
+
+#define UE_API EASECURVETOOL_API
+
+/** Serializes cubic bezier coordinates to cubic hermite */
+UCLASS(MinimalAPI)
+class UCubicBezierCurveSerializer : public UEaseCurveSerializer
+{
+	GENERATED_BODY()
+
+public:
+	//~ Begin UEaseCurveSerializer
+
+	virtual FText GetDisplayName() const override;
+	virtual FText GetDisplayTooltip() const override;
+
+	virtual bool IsFileExport() const override;
+	virtual bool SupportsExport() const override;
+	virtual bool Export(const FString& InFilePath, TSet<TWeakObjectPtr<UEaseCurveLibrary>> InWeakLibraries) override;
+
+	virtual bool IsFileImport() const override;
+	virtual bool SupportsImport() const override;
+    virtual bool Import(const FString& InFilePath, TSet<TWeakObjectPtr<UEaseCurveLibrary>> InWeakLibraries) override;
+
+	//~ End UEaseCurveSerializer
+};
+
+#undef UE_API
diff --git a/Engine/Plugins/Experimental/EaseCurveTool/Source/EaseCurveTool/Private/Serializers/CubicBezierPointEntryDialog.cpp b/Engine/Plugins/Experimental/EaseCurveTool/Source/EaseCurveTool/Private/Serializers/CubicBezierPointEntryDialog.cpp
new file mode 100644
--- /dev/null
+++ b/Engine/Plugins/Experimental/EaseCurveTool/Source/EaseCurveTool/Private/Serializers/CubicBezierPointEntryDialog.cpp
@@ -0,0 +1,392 @@
+﻿// Copyright Epic Games, Inc. All Rights Reserved.
+
+#include "Serializers/CubicBezierPointEntryDialog.h"
+#include "DetailLayoutBuilder.h"
+#include "EaseCurveLibrary.h"
+#include "EaseCurvePreset.h"
+#include "EaseCurveTangents.h"
+#include "EaseCurveToolSettings.h"
+#include "Editor.h"
+#include "Framework/Application/SlateApplication.h"
+#include "Widgets/Input/SButton.h"
+#include "Widgets/Input/SEditableTextBox.h"
+#include "Widgets/Input/SVectorInputBox.h"
+#include "Widgets/Layout/SBox.h"
+#include "Widgets/Layout/SUniformGridPanel.h"
+#include "Widgets/SBoxPanel.h"
+#include "Widgets/SEaseCurveCategoryComboBox.h"
+#include "Widgets/Text/STextBlock.h"
+
+#define LOCTEXT_NAMESPACE "CubicBezierPointEntryDialog"
+
+namespace UE::EaseCurveTool
+{
+
+bool FCubicBezierPointEntryDialog::Prompt(const TSet<TWeakObjectPtr<UEaseCurveLibrary>>& InWeakLibraries)
+{
+	if (!GEditor)
+	{
+		return false;
+	}
+
+	const TSharedRef<SWindow> Window = SNew(SWindow)
+		.Title(LOCTEXT("AddCubicBezierPreset", "Add Cubic Bezier Preset"))
+		.SupportsMinimize(false)
+		.SupportsMaximize(false)
+		.HasCloseButton(true)
+		.SizingRule(ESizingRule::Autosized)
+		.AutoCenter(EAutoCenter::PreferredWorkArea);
+
+	const TSharedRef<FCubicBezierPointEntryDialog> Dialog = MakeShared<FCubicBezierPointEntryDialog>();
+	Dialog->WeakLibraries = InWeakLibraries;
+	Dialog->Window = Window;
+
+	Window->SetContent(Dialog->BuildContent());
+
+	GEditor->EditorAddModalWindow(Window);
+
+	return Dialog->bHasAdded;
+}
+
+TSharedRef<SWidget> FCubicBezierPointEntryDialog::ConstructLabel(const FText& InLabel)
+{
+	return SNew(STextBlock)
+		.ColorAndOpacity(FStyleColors::Foreground)
+		.Text(InLabel);
+}
+
+TSharedRef<SWidget> FCubicBezierPointEntryDialog::BuildContent()
+{
+	using SNumericVectorInputBox2D = SNumericVectorInputBox<double, UE::Math::TVector2<double>, 2>;
+
+	static constexpr float Padding = 12.f;
+
+	return
+		SNew(SVerticalBox)
+		+ SVerticalBox::Slot()
+		.AutoHeight()
+		.Padding(Padding, Padding, Padding, 3.f)
+		[
+			SNew(STextBlock)
+			.ColorAndOpacity(FStyleColors::Foreground)
+			.Text(LOCTEXT("PresetCategory", "Preset Category"))
+		]
+		+ SVerticalBox::Slot()
+		.AutoHeight()
+		.Padding(Padding, 0.f, Padding, 6.f)
+		[
+			SNew(SEaseCurveCategoryComboBox, WeakLibraries)
+			.OnChanged_Lambda([this](const FText& InCategory, const ESelectInfo::Type InSelectInfo)
+				{
+					NewCategoryName = InCategory;
+					bLastAddedSuccessfully = false;
+				})
+		]
+		+ SVerticalBox::Slot()
+		.AutoHeight()
+		.Padding(Padding, 0.f, Padding, 3.f)
+		[
+			SNew(STextBlock)
+			.ColorAndOpacity(FStyleColors::Foreground)
+			.Text(LOCTEXT("PresetName", "Preset Name"))
+		]
+		+ SVerticalBox::Slot()
+		.AutoHeight()
+		.Padding(Padding, 0.f, Padding, Padding)
+		[
+			SNew(SEditableTextBox)
+			.OnTextChanged_Lambda([this](const FText& InText)
+				{
+					NewPresetName = InText;
+					bLastAddedSuccessfully = false;
+				})
+			.OnTextCommitted_Lambda([this](const FText& InText, ETextCommit::Type InCommitType)
+				{
+					NewPresetName = InText;
+					bLastAddedSuccessfully = false;
+				})
+		]
+		+ SVerticalBox::Slot()
+		.AutoHeight()
+		.Padding(Padding, 0.f, Padding, 6.f)
+		[
+			SNew(SHorizontalBox)
+			+ SHorizontalBox::Slot()
+			.AutoWidth()
+			.VAlign(VAlign_Center)
+			.Padding(0.f, 0.f, Padding, 0.f)
+			[
+				ConstructLabel(LOCTEXT("Point1", "Point 1"))
+			]
+			+ SHorizontalBox::Slot()
+			.HAlign(HAlign_Fill)
+			.VAlign(VAlign_Center)
+			[
+				SNew(SNumericVectorInputBox2D)
+				.bColorAxisLabels(true)
+				.AllowSpin(true)
+				.X_Lambda([this]()
+					{
+						return Point1.X;
+					})
+				.Y_Lambda([this]()
+					{
+						return Point1.Y;
+					})
+				.OnXChanged_Lambda([this](const double InNewValue)
+					{
+						Point1.X = InNewValue;
+						bLastAddedSuccessfully = false;
+					})
+				.OnYChanged_Lambda([this](const double InNewValue)
+					{
+						Point1.Y = InNewValue;
+						bLastAddedSuccessfully = false;
+					})
+				.OnXCommitted_Lambda([this](const double InNewValue, const ETextCommit::Type)
+					{
+						Point1.X = InNewValue;
+						bLastAddedSuccessfully = false;
+					})
+				.OnYCommitted_Lambda([this](const double InNewValue, const ETextCommit::Type)
+					{
+						Point1.Y = InNewValue;
+						bLastAddedSuccessfully = false;
+					})
+			]
+		]
+		+ SVerticalBox::Slot()
+		.AutoHeight()
+		.Padding(Padding, 0.f, Padding, 2.f)
+		[
+			SNew(SHorizontalBox)
+			+ SHorizontalBox::Slot()
+			.AutoWidth()
+			.VAlign(VAlign_Center)
+			.Padding(0.f, 0.f, Padding, 0.f)
+			[
+				ConstructLabel(LOCTEXT("Point2", "Point 2"))
+			]
+			+ SHorizontalBox::Slot()
+			.HAlign(HAlign_Fill)
+			.VAlign(VAlign_Center)
+			[
+				SNew(SNumericVectorInputBox2D)
+				.bColorAxisLabels(true)
+				.AllowSpin(true)
+				.X_Lambda([this]()
+					{
+						return Point2.X;
+					})
+				.Y_Lambda([this]()
+					{
+						return Point2.Y;
+					})
+				.OnXChanged_Lambda([this](const double InNewValue)
+					{
+						Point2.X = InNewValue;
+						bLastAddedSuccessfully = false;
+					})
+				.OnYChanged_Lambda([this](const double InNewValue)
+					{
+						Point2.Y = InNewValue;
+						bLastAddedSuccessfully = false;
+					})
+				.OnXCommitted_Lambda([this](const double InNewValue, const ETextCommit::Type)
+					{
+						Point2.X = InNewValue;
+						bLastAddedSuccessfully = false;
+					})
+				.OnYCommitted_Lambda([this](const double InNewValue, const ETextCommit::Type)
+					{
+						Point2.Y = InNewValue;
+						bLastAddedSuccessfully = false;
+					})
+			]
+		]
+		+ SVerticalBox::Slot()
+		.AutoHeight()
+		.HAlign(HAlign_Center)
+		.VAlign(VAlign_Center)
+		.Padding(Padding, Padding, Padding, 0.f)
+		[
+			SNew(STextBlock)
+			.Font(IDetailLayoutBuilder::GetDetailFont())
+			.ColorAndOpacity_Lambda([this]()
+				{
+					return StatusColor;
+				})
+			.Text_Lambda([this]()
+				{
+					return StatusText;
+				})
+		]
+		+ SVerticalBox::Slot()
+		.AutoHeight()
+		.Padding(Padding)
+		.HAlign(HAlign_Right)
+		[
+			SNew(SUniformGridPanel)
+			.SlotPadding(FMargin(4.f, 0.f))
+			+ SUniformGridPanel::Slot(0, 0)
+			.HAlign(HAlign_Fill)
+			.VAlign(VAlign_Center)
+			[
+				SNew(SButton)
+				.IsEnabled(this, &FCubicBezierPointEntryDialog::CanAdd)
+				.OnClicked(this, &FCubicBezierPointEntryDialog::HandleAddButtonClick)
+				.HAlign(HAlign_Center)
+				[
+					SNew(STextBlock)
+					.Text(LOCTEXT("AddButton", "Add"))
+				]
+			]
+			+ SUniformGridPanel::Slot(1, 0)
+			.HAlign(HAlign_Fill)
+			.VAlign(VAlign_Center)
+			[
+				SNew(SButton)
+				.OnClicked(this, &FCubicBezierPointEntryDialog::HandleCloseButtonClick)
+				.HAlign(HAlign_Center)
+				[
+					SNew(STextBlock)
+					.Text(LOCTEXT("CloseButton", "Close"))
+				]
+			]
+		];
+}
+
+bool FCubicBezierPointEntryDialog::CanAdd() const
+{
+	if (bLastAddedSuccessfully || NewPresetName.IsEmpty())
+	{
+		return false;
+	}
+
+	const FText CategoryToCheck = GetNewCategoryName();
+	if (CategoryToCheck.IsEmpty())
+	{
+		SetStatus(FStyleColors::Error, LOCTEXT("NoDefaultCategory", "No default category in tool settings!"));
+		return false;
+	}
+
+	const FText DefaultCategoryName = GetDefaultCategoryName();
+
+	bool bNoExistingCategory = false;
+	bool bPresetNameExists = false;
+
+	ForEachLibrary([this, &CategoryToCheck, &DefaultCategoryName, &bNoExistingCategory, &bPresetNameExists]
+		(const UEaseCurveLibrary* const InLibrary)
+		{
+			if (!InLibrary->DoesPresetCategoryExist(CategoryToCheck)
+				&& !CategoryToCheck.EqualToCaseIgnored(DefaultCategoryName))
+			{
+				bNoExistingCategory = true;
+				return false;
+			}
+
+			if (InLibrary->DoesPresetExist(FEaseCurvePresetHandle(CategoryToCheck, NewPresetName)))
+			{
+				bPresetNameExists = true;
+				return false;
+			}
+
+			return true;
+		});
+
+	if (bNoExistingCategory)
+	{
+		SetStatus(FStyleColors::Error, LOCTEXT("NoCategoryExists", "No category exists in library!"));
+	}
+	else if (bPresetNameExists)
+	{
+		SetStatus(FStyleColors::Error, LOCTEXT("PresetNameExists", "Preset name already exists in category!"));
+	}
+	else
+	{
+		SetStatus();
+	}
+
+	return !bNoExistingCategory && !bPresetNameExists;
+}
+
+FReply FCubicBezierPointEntryDialog::HandleAddButtonClick()
+{
+	if (!CanAdd())
+	{
+		return FReply::Handled();
+	}
+
+	const FText CategoryToAdd = GetNewCategoryName();
+
+	ForEachLibrary([this, &CategoryToAdd](UEaseCurveLibrary* const InLibrary)
+		{
+			const FEaseCurvePreset NewPreset(CategoryToAdd, NewPresetName
+				, FEaseCurveTangents::MakeFromCubicBezier({ Point1.X, Point1.Y, Point2.X, Point2.Y }));
+			InLibrary->AddPreset(NewPreset);
+			return true;
+		});
+
+	bHasAdded = true;
+	bLastAddedSuccessfully = true;
+
+	SetStatus(FStyleColors::Success, LOCTEXT("AddedSuccessfully", "Added successfully!"));
+	ResetDialog();
+
+	return FReply::Handled();
+}
+
+FReply FCubicBezierPointEntryDialog::HandleCloseButtonClick()
+{
+	if (Window.IsValid())
+	{
+		Window->RequestDestroyWindow();
+	}
+	return FReply::Handled();
+}
+
+void FCubicBezierPointEntryDialog::ForEachLibrary(const TFunction<bool(UEaseCurveLibrary*)> InFunctor) const
+{
+	for (const TWeakObjectPtr<UEaseCurveLibrary>& Library : WeakLibraries)
+	{
+		if (Library.IsValid())
+		{
+			if (!InFunctor(Library.Get()))
+			{
+				return;
+			}
+		}
+	}
+}
+
+void FCubicBezierPointEntryDialog::SetStatus(const FSlateColor& InColor, const FText& InText) const
+{
+	StatusColor = InColor;
+	StatusText = InText;
+}
+
+FText FCubicBezierPointEntryDialog::GetDefaultCategoryName() const
+{
+	if (const UEaseCurveToolSettings* const ToolSettings = GetDefault<UEaseCurveToolSettings>())
+	{
+		return ToolSettings->GetNewPresetCategory();
+	}
+	return FText::GetEmpty();
+}
+
+FText FCubicBezierPointEntryDialog::GetNewCategoryName() const
+{
+	return NewCategoryName.IsEmpty() ? GetDefaultCategoryName() : NewCategoryName;
+}
+
+void FCubicBezierPointEntryDialog::ResetDialog()
+{
+	NewPresetName = FText::GetEmpty();
+
+	Point1 = FVector2d::ZeroVector;
+	Point2 = FVector2d::ZeroVector;
+}
+
+} // namespace UE::SequenceNavigator
+
+#undef LOCTEXT_NAMESPACE
diff --git a/Engine/Plugins/Experimental/EaseCurveTool/Source/EaseCurveTool/Private/Serializers/CubicBezierPointEntryDialog.h b/Engine/Plugins/Experimental/EaseCurveTool/Source/EaseCurveTool/Private/Serializers/CubicBezierPointEntryDialog.h
new file mode 100644
--- /dev/null
+++ b/Engine/Plugins/Experimental/EaseCurveTool/Source/EaseCurveTool/Private/Serializers/CubicBezierPointEntryDialog.h
@@ -0,0 +1,65 @@
+﻿// Copyright Epic Games, Inc. All Rights Reserved.
+
+#pragma once
+
+#include "Internationalization/Text.h"
+#include "Math/MathFwd.h"
+#include "Misc/OptionalFwd.h"
+#include "Styling/SlateColor.h"
+#include "Styling/StyleColors.h"
+#include "Templates/SharedPointer.h"
+#include "UObject/WeakObjectPtrTemplatesFwd.h"
+
+class FReply;
+class SWidget;
+class SWindow;
+class UEaseCurveLibrary;
+class UObject;
+
+namespace UE::EaseCurveTool
+{
+
+class FCubicBezierPointEntryDialog : public TSharedFromThis<FCubicBezierPointEntryDialog>
+{
+public:
+	static bool Prompt(const TSet<TWeakObjectPtr<UEaseCurveLibrary>>& InWeakLibraries);
+
+protected:
+	TSharedRef<SWidget> ConstructLabel(const FText& InLabel);
+	TSharedRef<SWidget> BuildContent();
+
+	bool CanAdd() const;
+
+	FReply HandleAddButtonClick();
+	FReply HandleCloseButtonClick();
+
+	void ForEachLibrary(const TFunction<bool(UEaseCurveLibrary*)> InFunctor) const;
+
+	void SetStatus(const FSlateColor& InColor = FStyleColors::Foreground, const FText& InText = FText::GetEmpty()) const;
+
+	FText GetDefaultCategoryName() const;
+
+	/** @return The current selected category name.
+	 * If the category dropdown is set to "(Default)", this will return the default category from the tool settings. */
+	FText GetNewCategoryName() const;
+
+	void ResetDialog();
+
+	TSet<TWeakObjectPtr<UEaseCurveLibrary>> WeakLibraries;
+
+	TSharedPtr<SWindow> Window;
+
+	bool bHasAdded = false;
+
+	FText NewCategoryName;
+	FText NewPresetName;
+
+	FVector2d Point1 = FVector2d::ZeroVector;
+	FVector2d Point2 = FVector2d::ZeroVector;
+
+	mutable FSlateColor StatusColor;
+	mutable FText StatusText;
+	bool bLastAddedSuccessfully = false;
+};
+
+} // namespace UE::SequenceNavigator
diff --git a/Engine/Plugins/Experimental/EaseCurveTool/Source/EaseCurveTool/Private/Widgets/SEaseCurveCategoryComboBox.cpp b/Engine/Plugins/Experimental/EaseCurveTool/Source/EaseCurveTool/Private/Widgets/SEaseCurveCategoryComboBox.cpp
new file mode 100644
--- /dev/null
+++ b/Engine/Plugins/Experimental/EaseCurveTool/Source/EaseCurveTool/Private/Widgets/SEaseCurveCategoryComboBox.cpp
@@ -0,0 +1,108 @@
+﻿// Copyright Epic Games, Inc. All Rights Reserved.
+
+#include "SEaseCurveCategoryComboBox.h"
+#include "EaseCurveLibrary.h"
+#include "Styling/StyleColors.h"
+#include "Widgets/Input/SButton.h"
+#include "Widgets/Input/SComboBox.h"
+#include "Widgets/Layout/SBox.h"
+#include "Widgets/SBoxPanel.h"
+
+#define LOCTEXT_NAMESPACE "SEaseCurveCategoryComboBox"
+
+namespace UE::EaseCurveTool
+{
+
+void SEaseCurveCategoryComboBox::Construct(const FArguments& InArgs, const TWeakObjectPtr<UEaseCurveLibrary>& InWeakLibrary)
+{
+	Construct(InArgs, TSet<TWeakObjectPtr<UEaseCurveLibrary>>{ InWeakLibrary });
+}
+
+void SEaseCurveCategoryComboBox::Construct(const FArguments& InArgs, const TSet<TWeakObjectPtr<UEaseCurveLibrary>>& InWeakLibraries)
+{
+	WeakLibraries = InWeakLibraries;
+
+	OnCategoryChanged = InArgs._OnChanged;
+
+	if (InArgs._UseDefault)
+	{
+		DefaultCategory = MakeShared<FText>(LOCTEXT("DefaultCategory", "(Default)"));
+		AllCategories.Add(DefaultCategory);
+
+		SelectedCategory = DefaultCategory;
+	}
+
+	for (auto Iter = WeakLibraries.CreateIterator(); Iter; ++Iter)
+	{
+		const TWeakObjectPtr<UEaseCurveLibrary> Library = *Iter;
+		if (Library.IsValid())
+		{
+			for (const FText& Category : Library->GetCategories())
+			{
+				AllCategories.Add(MakeShared<FText>(Category));
+			}
+		}
+		else
+		{
+			Iter.RemoveCurrent();
+		}
+	}
+
+	AllCategories.Sort([this](const TSharedPtr<FText>& InA, const TSharedPtr<FText>& InB)
+		{
+			if (DefaultCategory.IsValid() && InA == DefaultCategory)
+			{
+				return true;
+			}
+			return InA->CompareToCaseIgnored(*InB) < 0;
+		});
+
+	ChildSlot
+	[
+		SNew(SComboBox<TSharedPtr<FText>>)
+		.OptionsSource(&AllCategories)
+		.OnGenerateWidget_Lambda([](const TSharedPtr<FText> InItem)
+			{
+				return SNew(STextBlock)
+					.ColorAndOpacity(FStyleColors::Foreground)
+					.Text(*InItem);
+			})
+		.OnSelectionChanged(this, &SEaseCurveCategoryComboBox::HandleCategoryChanged)
+		.Content()
+		[
+			SNew(STextBlock)
+			.Text(this, &SEaseCurveCategoryComboBox::GetSelectedCategory)
+		]
+	];
+}
+
+FText SEaseCurveCategoryComboBox::GetSelectedCategory() const
+{
+	if (SelectedCategory.IsValid())
+	{
+		return *SelectedCategory;
+	}
+	return FText::GetEmpty();
+}
+
+TSharedPtr<FText> SEaseCurveCategoryComboBox::GetSelectedCategoryItem() const
+{
+	return SelectedCategory;
+}
+
+void SEaseCurveCategoryComboBox::HandleCategoryChanged(const TSharedPtr<FText> InNewCategory, const ESelectInfo::Type InSelectInfo)
+{
+	if (WeakLibraries.IsEmpty())
+	{
+		return;
+	}
+
+	SelectedCategory = InNewCategory;
+
+	const FText CategoryName = (SelectedCategory == DefaultCategory) ? FText::GetEmpty() : *SelectedCategory;
+	OnCategoryChanged.ExecuteIfBound(CategoryName, InSelectInfo);
+}
+
+} // namespace UE::EaseCurveTool
+
+#undef LOCTEXT_NAMESPACE
diff --git a/Engine/Plugins/Experimental/EaseCurveTool/Source/EaseCurveTool/Private/Widgets/SEaseCurveCategoryComboBox.h b/Engine/Plugins/Experimental/EaseCurveTool/Source/EaseCurveTool/Private/Widgets/SEaseCurveCategoryComboBox.h
new file mode 100644
--- /dev/null
+++ b/Engine/Plugins/Experimental/EaseCurveTool/Source/EaseCurveTool/Private/Widgets/SEaseCurveCategoryComboBox.h
@@ -0,0 +1,43 @@
+﻿// Copyright Epic Games, Inc. All Rights Reserved.
+
+#pragma once
+
+#include "Widgets/SCompoundWidget.h"
+
+class FText;
+class UEaseCurveLibrary;
+
+namespace UE::EaseCurveTool
+{
+
+/** A combo box widget to display and select a category from a list of all categories that exist inside a library  */
+class SEaseCurveCategoryComboBox : public SCompoundWidget
+{
+public:
+	DECLARE_DELEGATE_TwoParams(FOnCategoryChanged, const FText& /*InNewCategory*/, const ESelectInfo::Type /*InSelectInfo*/)
+
+	SLATE_BEGIN_ARGS(SEaseCurveCategoryComboBox) {}
+		SLATE_ARGUMENT(bool, UseDefault)
+		SLATE_EVENT(FOnCategoryChanged, OnChanged)
+	SLATE_END_ARGS()
+
+	void Construct(const FArguments& InArgs, const TWeakObjectPtr<UEaseCurveLibrary>& InWeakLibrary);
+	void Construct(const FArguments& InArgs, const TSet<TWeakObjectPtr<UEaseCurveLibrary>>& InWeakLibraries);
+
+	FText GetSelectedCategory() const;
+
+protected:
+	TSharedPtr<FText> GetSelectedCategoryItem() const;
+	void HandleCategoryChanged(const TSharedPtr<FText> InNewCategory, const ESelectInfo::Type InSelectInfo);
+
+	TSet<TWeakObjectPtr<UEaseCurveLibrary>> WeakLibraries;
+
+	FOnCategoryChanged OnCategoryChanged;
+
+	TArray<TSharedPtr<FText>> AllCategories;
+
+	TSharedPtr<FText> DefaultCategory;
+	TSharedPtr<FText> SelectedCategory;
+};
+
+} // namespace UE::EaseCurveTool
