[SequenceNavigator] Change or remove "Navigation Tool" comment references to "Sequence Navigator" for outside modules

diff --git a/Engine/Plugins/Experimental/AvalancheSequenceNavigator/Source/AvalancheSequenceNavigator/Private/AvalancheSequenceNavigatorModule.h b/Engine/Plugins/Experimental/AvalancheSequenceNavigator/Source/AvalancheSequenceNavigator/Private/AvalancheSequenceNavigatorModule.h
--- a/Engine/Plugins/Experimental/AvalancheSequenceNavigator/Source/AvalancheSequenceNavigator/Private/AvalancheSequenceNavigatorModule.h
+++ b/Engine/Plugins/Experimental/AvalancheSequenceNavigator/Source/AvalancheSequenceNavigator/Private/AvalancheSequenceNavigatorModule.h
@@ -33,6 +33,6 @@ private:
 
 	FDelegateHandle AvaSequencerCreatedHandle;
 
-	/** The provider that supplies data and extends the Navigation Tool */
+	/** The provider that supplies data and extends the Sequence Navigator */
 	TSharedPtr<UE::AvaSequencer::FAvaNavigationToolProvider> NavigationToolProvider;
 };
diff --git a/Engine/Plugins/Experimental/CinematicSequenceNavigator/Source/CinematicSequenceNavigator/Private/CinematicSequenceNavigatorModule.h b/Engine/Plugins/Experimental/CinematicSequenceNavigator/Source/CinematicSequenceNavigator/Private/CinematicSequenceNavigatorModule.h
--- a/Engine/Plugins/Experimental/CinematicSequenceNavigator/Source/CinematicSequenceNavigator/Private/CinematicSequenceNavigatorModule.h
+++ b/Engine/Plugins/Experimental/CinematicSequenceNavigator/Source/CinematicSequenceNavigator/Private/CinematicSequenceNavigatorModule.h
@@ -28,6 +28,6 @@ private:
 	FDelegateHandle SequencerCreatedHandle;
 	FDelegateHandle SequencerClosedHandle;
 
-	/** The provider that supplies data and extends the Navigation Tool */
+	/** The provider that supplies data and extends the Sequence Navigator */
 	TSharedPtr<UE::CineAssemblyTools::FCinematicNavigationToolProvider> NavigationToolProvider;
 };
diff --git a/Engine/Plugins/Experimental/EaseCurveTool/Source/EaseCurveTool/Private/EaseCurveToolExtender.h b/Engine/Plugins/Experimental/EaseCurveTool/Source/EaseCurveTool/Private/EaseCurveToolExtender.h
--- a/Engine/Plugins/Experimental/EaseCurveTool/Source/EaseCurveTool/Private/EaseCurveToolExtender.h
+++ b/Engine/Plugins/Experimental/EaseCurveTool/Source/EaseCurveTool/Private/EaseCurveToolExtender.h
@@ -22,7 +22,7 @@ namespace UE::EaseCurveTool
 class FEaseCurveTool;
 
 /**
- * Singleton class for extending and managing Ease Curve Tool instances.
+ * Singleton class for extending and managing tool instances.
  */
 class FEaseCurveToolExtender : IModularFeature
 {
@@ -36,17 +36,17 @@ public:
 
 	/**
 	 * Retrieves the unique name for the provided Sequencer instance.
-	 * The Ease Curve Tool instance name is based on its associated sequencer settings
+	 * The tool instance name is based on its associated sequencer settings
 	 * object name. This provides a unique instance name for each Sequencer type.
 	 * @param InSequencer The Sequencer instance for which the name is being retrieved.
-	 * @return The name of the Ease Curve Tool instance if available; otherwise, NAME_None.
+	 * @return The name of the tool instance if available; otherwise, NAME_None.
 	 */
 	UE_API static FName GetToolInstanceId(const ISequencer& InSequencer);
 	UE_API static FName GetToolInstanceId(const FEaseCurveTool& InTool);
 
 	UE_API static TSharedPtr<FEaseCurveTool> GetToolInstance(const FName InToolId);
 
-	/** Find the Navigation Tool instance associated with the given Sequencer instance */
+	/** Find the tool instance associated with the given Sequencer instance */
 	UE_API static TSharedPtr<FEaseCurveTool> FindToolInstance(const TSharedRef<ISequencer>& InSequencer);
 
 	UE_API static TSharedPtr<FEaseCurveTool> FindToolInstanceByCurveEditor(const TSharedRef<FCurveEditor>& InCurveEditor);
@@ -84,7 +84,7 @@ private:
 
 	FDelegateHandle SequencerCreatedHandle;
 
-	/** Registered Sequencer types to Ease Curve Tool instances */
+	/** Registered Sequencer types to the tool instances */
 	TMap<FName, FEaseCurveToolInstance> ToolInstances;
 };
 
diff --git a/Engine/Plugins/Experimental/EaseCurveTool/Source/EaseCurveTool/Private/EaseCurveToolTab.h b/Engine/Plugins/Experimental/EaseCurveTool/Source/EaseCurveTool/Private/EaseCurveToolTab.h
--- a/Engine/Plugins/Experimental/EaseCurveTool/Source/EaseCurveTool/Private/EaseCurveToolTab.h
+++ b/Engine/Plugins/Experimental/EaseCurveTool/Source/EaseCurveTool/Private/EaseCurveToolTab.h
@@ -30,13 +30,13 @@ public:
 	void Init();
 	void Shutdown();
 
-	/** @return True if the USequencerSettings setting for the Navigation Tool is set to visible */
+	/** @return True if the USequencerSettings setting for the tool is set to visible */
 	bool ShouldShowToolTab() const;
 
-	/** @return True if the Navigation Tool tab is visible */
+	/** @return True if the tool tab is visible */
 	bool IsToolTabVisible() const;
 
-	/** Shows or hides the Navigation Tool tab and toggles the state in the saved settings */
+	/** Shows or hides the tool tab and toggles the state in the saved settings */
 	void ShowHideToolTab(const bool bInVisible);
 
 	void ToggleToolTabVisible();
diff --git a/Engine/Plugins/Experimental/EaseCurveTool/Source/EaseCurveTool/Private/Serializers/CubicBezierPointEntryDialog.cpp b/Engine/Plugins/Experimental/EaseCurveTool/Source/EaseCurveTool/Private/Serializers/CubicBezierPointEntryDialog.cpp
--- a/Engine/Plugins/Experimental/EaseCurveTool/Source/EaseCurveTool/Private/Serializers/CubicBezierPointEntryDialog.cpp
+++ b/Engine/Plugins/Experimental/EaseCurveTool/Source/EaseCurveTool/Private/Serializers/CubicBezierPointEntryDialog.cpp
@@ -387,6 +387,6 @@ void FCubicBezierPointEntryDialog::ResetDialog()
 	Point2 = FVector2d::ZeroVector;
 }
 
-} // namespace UE::SequenceNavigator
+} // namespace UE::EaseCurveTool
 
 #undef LOCTEXT_NAMESPACE
diff --git a/Engine/Plugins/Experimental/EaseCurveTool/Source/EaseCurveTool/Private/Serializers/CubicBezierPointEntryDialog.h b/Engine/Plugins/Experimental/EaseCurveTool/Source/EaseCurveTool/Private/Serializers/CubicBezierPointEntryDialog.h
--- a/Engine/Plugins/Experimental/EaseCurveTool/Source/EaseCurveTool/Private/Serializers/CubicBezierPointEntryDialog.h
+++ b/Engine/Plugins/Experimental/EaseCurveTool/Source/EaseCurveTool/Private/Serializers/CubicBezierPointEntryDialog.h
@@ -62,4 +62,4 @@ protected:
 	bool bLastAddedSuccessfully = false;
 };
 
-} // namespace UE::SequenceNavigator
+} // namespace UE::EaseCurveTool
diff --git a/Engine/Plugins/Experimental/LevelSequenceNavigatorBridge/Source/LevelSequenceNavigatorBridge/Private/LevelSequenceNavigatorBridgeModule.h b/Engine/Plugins/Experimental/LevelSequenceNavigatorBridge/Source/LevelSequenceNavigatorBridge/Private/LevelSequenceNavigatorBridgeModule.h
--- a/Engine/Plugins/Experimental/LevelSequenceNavigatorBridge/Source/LevelSequenceNavigatorBridge/Private/LevelSequenceNavigatorBridgeModule.h
+++ b/Engine/Plugins/Experimental/LevelSequenceNavigatorBridge/Source/LevelSequenceNavigatorBridge/Private/LevelSequenceNavigatorBridgeModule.h
@@ -24,6 +24,6 @@ private:
 	FDelegateHandle SequencerCreatedHandle;
 	FDelegateHandle SequencerClosedHandle;
 
-	/** The provider that supplies data and extends the Navigation Tool */
+	/** The provider that supplies data and extends the Sequence Navigator */
 	TSharedPtr<FLevelSequenceNavigationToolProvider> NavigationToolProvider;
 };
diff --git a/Engine/Plugins/Experimental/SequenceNavigator/Source/SequenceNavigator/Private/NavigationToolTab.h b/Engine/Plugins/Experimental/SequenceNavigator/Source/SequenceNavigator/Private/NavigationToolTab.h
--- a/Engine/Plugins/Experimental/SequenceNavigator/Source/SequenceNavigator/Private/NavigationToolTab.h
+++ b/Engine/Plugins/Experimental/SequenceNavigator/Source/SequenceNavigator/Private/NavigationToolTab.h
@@ -16,13 +16,13 @@ namespace UE::SequenceNavigator
 
 class INavigationTool;
 
-/** Manages a Navigation Tool tab, including its visible state settings in USequencerSettings */
+/** Manages a tool instance tab, including its visible state settings in USequencerSettings */
 class FNavigationToolTab : public TSharedFromThis<FNavigationToolTab>
 {
 public:
 	DECLARE_MULTICAST_DELEGATE_OneParam(FOnVisibilityChanged, const bool /*bInVisible*/);
 
-	/** @return Returns the auto generated tool tab id using the passed in ISequencer.
+	/** @return Returns the auto generated tool instance tab id using the passed in ISequencer.
 	 * Tool tab ids are based on the sequencer instance settings name. */
 	static FName StaticToolTabId(const ISequencer& InSequencer);
 
@@ -31,13 +31,13 @@ public:
 	void Init();
 	void Shutdown();
 
-	/** @return True if the USequencerSettings setting for the Navigation Tool is set to visible */
+	/** @return True if the USequencerSettings setting for the tool instance is set to visible */
 	bool ShouldShowToolTab() const;
 
-	/** @return True if the Navigation Tool tab is visible */
+	/** @return True if the tool instance tab is visible */
 	bool IsToolTabVisible() const;
 
-	/** Shows or hides the Navigation Tool tab and toggles the state in the saved settings */
+	/** Shows or hides the tool instance tab and toggles the state in the saved settings */
 	void ShowHideToolTab(const bool bInVisible);
 
 	void ToggleToolTabVisible();
diff --git a/Engine/Plugins/VirtualProduction/Avalanche/Source/Avalanche/Public/AvaScene.h b/Engine/Plugins/VirtualProduction/Avalanche/Source/Avalanche/Public/AvaScene.h
--- a/Engine/Plugins/VirtualProduction/Avalanche/Source/Avalanche/Public/AvaScene.h
+++ b/Engine/Plugins/VirtualProduction/Avalanche/Source/Avalanche/Public/AvaScene.h
@@ -180,7 +180,7 @@ protected:
 	UPROPERTY()
 	bool bAutoStartMode = true;
 
-	/** Saved Navigation Tool state to restore */
+	/** Saved Sequence Navigator state to restore */
 	UPROPERTY()
 	FNavigationToolSaveState NavigationToolState;
 #endif
diff --git a/Engine/Source/Editor/Sequencer/Public/Filters/SequencerFilterBase.h b/Engine/Source/Editor/Sequencer/Public/Filters/SequencerFilterBase.h
--- a/Engine/Source/Editor/Sequencer/Public/Filters/SequencerFilterBase.h
+++ b/Engine/Source/Editor/Sequencer/Public/Filters/SequencerFilterBase.h
@@ -7,7 +7,7 @@
 #include "Templates/SharedPointer.h"
 #include "MovieSceneSequence.h"
 
-/** Sequencer base filter for all Sequencer filters (Sequencer, Navigation Tool, etc.) */
+/** Sequencer base filter for all Sequencer filters (Sequencer and Sequence Navigator) */
 template<typename InFilterType>
 class FSequencerFilterBase
 	: public FFilterBase<InFilterType>
diff --git a/Engine/Source/Editor/Sequencer/Public/SequencerSettings.h b/Engine/Source/Editor/Sequencer/Public/SequencerSettings.h
--- a/Engine/Source/Editor/Sequencer/Public/SequencerSettings.h
+++ b/Engine/Source/Editor/Sequencer/Public/SequencerSettings.h
@@ -572,9 +572,9 @@ public:
 	/** Sets how the asset thumbnail is captured when the sequence is saved.  */
 	UE_API void SetThumbnailCaptureSettings(const FSequencerThumbnailCaptureSettings& InNewValue);
 
-	/** @return True if the Navigation Tool is visible */
+	/** @return True if the Sequence Navigator is visible */
 	bool IsNavigationToolVisible() const { return bNavigationToolVisible; }
-	/** Sets the visibility of the Navigation Tool */
+	/** Sets the visibility of the Sequence Navigator */
 	UE_API void SetNavigationToolVisible(const bool bInVisible);
 
 protected:
@@ -920,8 +920,8 @@ protected:
 	UPROPERTY(config, EditAnywhere, Category = General, meta = (EditCondition="ShouldShowThumbnailCaptureSettings()", EditConditionHides))
 	FSequencerThumbnailCaptureSettings ThumbnailCaptureSettings;
 
-	/** Last saved visibility of the Navigation Tool to restore after closed */
-	UPROPERTY(config, EditAnywhere, Category = General)
+	/** Last saved visibility of the Sequence Navigator to restore after closed */
+	UPROPERTY(config, EditAnywhere, Category = General, meta = (DisplayName = "Sequence Navigator Visible"))
 	bool bNavigationToolVisible;
 
 	UFUNCTION()
diff --git a/Engine/Source/Runtime/MovieScene/Public/EditorSequenceNavigationDefs.h b/Engine/Source/Runtime/MovieScene/Public/EditorSequenceNavigationDefs.h
--- a/Engine/Source/Runtime/MovieScene/Public/EditorSequenceNavigationDefs.h
+++ b/Engine/Source/Runtime/MovieScene/Public/EditorSequenceNavigationDefs.h
@@ -12,7 +12,7 @@ enum class ENavigationToolItemFlags : uint8
 	None = 0,
 	/** Whether item should get the underlying UObject, ignoring if it's pending kill */
 	IgnorePendingKill = 1 << 0,
-	/** Item pending removal from the Navigation Tool */
+	/** Item pending removal */
 	PendingRemoval = 1 << 1,
 	/** Whether the item is in expanded state to show its child items */
 	Expanded = 1 << 2,
@@ -148,7 +148,7 @@ struct FNavigationToolViewSaveState
 {
 	GENERATED_BODY()
 
-	/** Items specific to this Navigation Tool Instance, rather than being shared across Navigation Tools (e.g. Expansion flags) */
+	/** Items specific to this tool instance, rather than being shared across all tool instances (e.g. Expansion flags) */
 	UPROPERTY()
 	TMap<FString, ENavigationToolItemFlags> ViewItemFlags;
 
diff --git a/Engine/Source/Runtime/MovieScene/Public/MovieScene.h b/Engine/Source/Runtime/MovieScene/Public/MovieScene.h
--- a/Engine/Source/Runtime/MovieScene/Public/MovieScene.h
+++ b/Engine/Source/Runtime/MovieScene/Public/MovieScene.h
@@ -131,7 +131,7 @@ struct FMovieSceneEditorData
 	UPROPERTY()
 	FFloatRange ViewRange_DEPRECATED;
 
-	/** Saved Navigation Tool state to restore */
+	/** Saved Sequence Navigator state to restore */
 	UPROPERTY()
 	FNavigationToolSaveState NavigationToolState;
 
