[EaseCurveTool] Change modular feature name to "EaseCurveTool"

diff --git a/Engine/Plugins/Animation/EaseCurveTool/Source/EaseCurveTool/Private/EaseCurveToolExtender.cpp b/Engine/Plugins/Animation/EaseCurveTool/Source/EaseCurveTool/Private/EaseCurveToolExtender.cpp
--- a/Engine/Plugins/Animation/EaseCurveTool/Source/EaseCurveTool/Private/EaseCurveToolExtender.cpp
+++ b/Engine/Plugins/Animation/EaseCurveTool/Source/EaseCurveTool/Private/EaseCurveToolExtender.cpp
@@ -137,7 +137,7 @@ void FEaseCurveToolExtender::OnSequencerClosed(const TSharedRef<ISequencer> InSe
 
 FName FEaseCurveToolExtender::GetModularFeatureName()
 {
-	static FName FeatureName = TEXT("NavigationTool");
+	static FName FeatureName = TEXT("EaseCurveTool");
 	return FeatureName;
 }
 
