Fix warning for missing editor resource
- `CORE_BOX_BRUSH` is the one set to the appropriate path in this case

#Synced-CL 48299257

diff --git a/Engine/Plugins/Animation/EaseCurveTool/Source/EaseCurveTool/Private/EaseCurveStyle.cpp b/Engine/Plugins/Animation/EaseCurveTool/Source/EaseCurveTool/Private/EaseCurveStyle.cpp
--- a/Engine/Plugins/Animation/EaseCurveTool/Source/EaseCurveTool/Private/EaseCurveStyle.cpp
+++ b/Engine/Plugins/Animation/EaseCurveTool/Source/EaseCurveTool/Private/EaseCurveStyle.cpp
@@ -74,8 +74,8 @@ FEaseCurveStyle::FEaseCurveStyle()
 
 	FButtonStyle MenuHoverHintOnly = FButtonStyle()
 		.SetNormal(FSlateNoResource())
-		.SetHovered(BOX_BRUSH("Common/ButtonHoverHint", FMargin(), FStyleColors::Highlight.GetSpecifiedColor()))
-		.SetPressed(BOX_BRUSH("Common/ButtonHoverHint", FMargin(), FStyleColors::Highlight.GetSpecifiedColor()))
+		.SetHovered(CORE_BOX_BRUSH("Common/ButtonHoverHint", FMargin(), FStyleColors::Highlight.GetSpecifiedColor()))
+		.SetPressed(CORE_BOX_BRUSH("Common/ButtonHoverHint", FMargin(), FStyleColors::Highlight.GetSpecifiedColor()))
 		.SetNormalPadding(FMargin(0.f,0.f,0.f,1.f))
 		.SetPressedPadding(FMargin(0.f,1.f,0.f,0.f));
 	Set("MenuHoverHintOnly", MenuHoverHintOnly);
