[Sequencer] Ease Curve Tool:
- Remove unneeded mapping of command actions for undo/redo

#Synced-CL 48739784

diff --git a/Engine/Plugins/Animation/EaseCurveTool/Source/EaseCurveTool/Private/EaseCurveTool.cpp b/Engine/Plugins/Animation/EaseCurveTool/Source/EaseCurveTool/Private/EaseCurveTool.cpp
--- a/Engine/Plugins/Animation/EaseCurveTool/Source/EaseCurveTool/Private/EaseCurveTool.cpp
+++ b/Engine/Plugins/Animation/EaseCurveTool/Source/EaseCurveTool/Private/EaseCurveTool.cpp
@@ -594,28 +594,12 @@ void FEaseCurveTool::EndTransaction() const
 	}
 }
 
-void FEaseCurveTool::UndoAction()
-{
-	if (GEditor && GEditor->UndoTransaction())
-	{
-		UpdateEaseCurveFromKeySelections();
-	}
-}
-
-void FEaseCurveTool::RedoAction()
-{
-	if (GEditor && GEditor->RedoTransaction())
-	{
-		UpdateEaseCurveFromKeySelections();
-	}
-}
-
-void FEaseCurveTool::PostUndo(bool bInSuccess)
+void FEaseCurveTool::PostUndo(const bool bInSuccess)
 {
 	UpdateEaseCurveFromKeySelections();
 }
 
-void FEaseCurveTool::PostRedo(bool bInSuccess)
+void FEaseCurveTool::PostRedo(const bool bInSuccess)
 {
 	UpdateEaseCurveFromKeySelections();
 }
@@ -755,12 +739,6 @@ void FEaseCurveTool::BindCommands()
 
 	CommandList = MakeShared<FUICommandList>();
 
-	CommandList->MapAction(FGenericCommands::Get().Undo
-		, FExecuteAction::CreateSP(this, &FEaseCurveTool::UndoAction));
-	
-	CommandList->MapAction(FGenericCommands::Get().Redo
-		, FExecuteAction::CreateSP(this, &FEaseCurveTool::RedoAction));
-
 	CommandList->MapAction(EaseCurveToolCommands.OpenToolSettings
 		, FExecuteAction::CreateSP(this, &FEaseCurveTool::OpenToolSettings));
 
diff --git a/Engine/Plugins/Animation/EaseCurveTool/Source/EaseCurveTool/Private/EaseCurveTool.h b/Engine/Plugins/Animation/EaseCurveTool/Source/EaseCurveTool/Private/EaseCurveTool.h
--- a/Engine/Plugins/Animation/EaseCurveTool/Source/EaseCurveTool/Private/EaseCurveTool.h
+++ b/Engine/Plugins/Animation/EaseCurveTool/Source/EaseCurveTool/Private/EaseCurveTool.h
@@ -128,9 +128,6 @@ public:
 	void BeginTransaction(const FText& InDescription) const;
 	void EndTransaction() const;
 
-	void UndoAction();
-	void RedoAction();
-
 	void OpenToolSettings() const;
 
 	FFrameRate GetTickResolution() const;
@@ -150,8 +147,8 @@ public:
 	//~ End FGCObject
 
 	//~ Begin FEditorUndoClient
-	virtual void PostUndo(bool bInSuccess) override;
-	virtual void PostRedo(bool bInSuccess) override;
+	virtual void PostUndo(const bool bInSuccess) override;
+	virtual void PostRedo(const bool bInSuccess) override;
 	//~ End FEditorUndoClient
 
 	//~ Begin FTickableEditorObject
