Ease Curve Selection improvement

#rnx

#Synced-CL 50096176

diff --git a/Engine/Plugins/Animation/EaseCurveTool/Source/EaseCurveTool/Private/EaseCurveKeySelection.cpp b/Engine/Plugins/Animation/EaseCurveTool/Source/EaseCurveTool/Private/EaseCurveKeySelection.cpp
--- a/Engine/Plugins/Animation/EaseCurveTool/Source/EaseCurveTool/Private/EaseCurveKeySelection.cpp
+++ b/Engine/Plugins/Animation/EaseCurveTool/Source/EaseCurveTool/Private/EaseCurveKeySelection.cpp
@@ -244,9 +244,9 @@ void FEaseCurveKeySelection::SetSelectionError()
 	}
 	else
 	{
-		for (const TPair<FName, FEaseCurveChannelKeyCache>& KeyDataPair : SequencerKeyData)
+		for (const auto& [ChannelModel, KeyCache] : SequencerKeyData)
 		{
-			LoopKeyHandles(KeyDataPair.Value.KeyHandles);
+			LoopKeyHandles(KeyCache.KeyHandles);
 		}
 	}
 
@@ -298,12 +298,7 @@ bool FEaseCurveKeySelection::ProcessSequencerKeySelection(const TSharedPtr<FSequ
 	}
 
 	const FMovieSceneChannelHandle& ChannelHandle = KeyArea->GetChannel();
-
-	const FString PathName = IOutlinerExtension::GetPathName(ChannelModel);
-	const FString ChannelName = ChannelModel->GetChannelName().ToString();
-	const FName FullPathMapKey = FName(PathName + ChannelName);
-
-	FEaseCurveChannelKeyCache& Entry = SequencerKeyData.FindOrAdd(FullPathMapKey);
+	FEaseCurveChannelKeyCache& Entry = SequencerKeyData.FindOrAdd(ChannelModel);
 	Entry.Section = ChannelModel->GetSection();
 	Entry.ChannelHandle = ChannelHandle;
 	Entry.KeyHandles.Add(InKey, FindErrorForKey(ChannelHandle, InKey));
@@ -366,10 +361,9 @@ void FEaseCurveKeySelection::ForEachEaseableKey(const FProcessKeySignature& InCa
 		return;
 	}
 
-	for (const TPair<FName, FEaseCurveChannelKeyCache>& ChannelEntry : SequencerKeyData)
+	for (const auto& [ChannelModel, KeyCache] : SequencerKeyData)
 	{
-		if (!ForEachChannelKey(ChannelEntry.Value.ChannelHandle
-			, ChannelEntry.Value.KeyHandles, ChannelEntry.Value.Section, InCallable))
+		if (!ForEachChannelKey(KeyCache.ChannelHandle, KeyCache.KeyHandles, KeyCache.Section, InCallable))
 		{
 			return;
 		}
diff --git a/Engine/Plugins/Animation/EaseCurveTool/Source/EaseCurveTool/Private/EaseCurveKeySelection.h b/Engine/Plugins/Animation/EaseCurveTool/Source/EaseCurveTool/Private/EaseCurveKeySelection.h
--- a/Engine/Plugins/Animation/EaseCurveTool/Source/EaseCurveTool/Private/EaseCurveKeySelection.h
+++ b/Engine/Plugins/Animation/EaseCurveTool/Source/EaseCurveTool/Private/EaseCurveKeySelection.h
@@ -6,6 +6,7 @@
 #include "Channels/MovieSceneChannelData.h"
 #include "Channels/MovieSceneChannelHandle.h"
 #include "EaseCurveToolUtils.h"
+#include "MVVM/ViewModelPtr.h"
 #include "EaseCurveKeySelection.generated.h"
 
 class FCurveModel;
@@ -17,6 +18,7 @@ struct FMovieSceneChannelHandle;
 
 namespace UE::Sequencer
 {
+	class FChannelModel;
 	class FSequencerSelection;
 }
 
@@ -215,7 +217,7 @@ protected:
 		return true;
 	}
 
-	TMap<FName, FEaseCurveChannelKeyCache> SequencerKeyData;
+	TMap<UE::Sequencer::TViewModelPtr<UE::Sequencer::FChannelModel>, FEaseCurveChannelKeyCache> SequencerKeyData;
 	TMap<FCurveModelID, FEaseCurveChannelKeyCache> CurveEditorKeyData;
 
 	int32 TotalSelectedKeys = 0;
