Curve Editor: Fix ease curve tool pop up menu location, otherwise it can go offscreen when invoking from a collapsed menu bar

#rnx

#Synced-CL 52584676

diff --git a/Engine/Plugins/Animation/EaseCurveTool/Source/EaseCurveTool/Private/EaseCurveEditorExtension.cpp b/Engine/Plugins/Animation/EaseCurveTool/Source/EaseCurveTool/Private/EaseCurveEditorExtension.cpp
--- a/Engine/Plugins/Animation/EaseCurveTool/Source/EaseCurveTool/Private/EaseCurveEditorExtension.cpp
+++ b/Engine/Plugins/Animation/EaseCurveTool/Source/EaseCurveTool/Private/EaseCurveEditorExtension.cpp
@@ -1,4 +1,4 @@
-﻿// Copyright Epic Games, Inc. All Rights Reserved.
+// Copyright Epic Games, Inc. All Rights Reserved.
 
 #include "EaseCurveEditorExtension.h"
 #include "EaseCurvePreset.h"
@@ -92,9 +92,7 @@ void FEaseCurveEditorExtension::PopupToolbarMenu()
 		return;
 	}
 
-	const FGeometry ButtonGeometry = ButtonWidget->GetCachedGeometry();
-	const FVector2D SummonLocation = FVector2D(ButtonGeometry.AbsolutePosition.X,
-		ButtonGeometry.AbsolutePosition.Y + ButtonGeometry.GetAbsoluteSize().Y);
+	const FVector2D SummonLocation = FSlateApplication::Get().GetCursorPos();
 
 	ButtonMenu = MakeShared<FCurveEditorToolBarMenu>(GetToolInstance());
 
